Copied cl 9439572

MDLImporter: Importing MDL file crashes UE4

diff --git a/Engine/Plugins/Enterprise/MDLImporter/Source/MDLImporter/Private/MDLMapHandler.cpp b/Engine/Plugins/Enterprise/MDLImporter/Source/MDLImporter/Private/MDLMapHandler.cpp
--- a/Engine/Plugins/Enterprise/MDLImporter/Source/MDLImporter/Private/MDLMapHandler.cpp
+++ b/Engine/Plugins/Enterprise/MDLImporter/Source/MDLImporter/Private/MDLMapHandler.cpp
@@ -103,7 +103,15 @@ bool FMDLMapHandler::Import(const FString& MapName, bool bIsTexture, Mdl::FBakeP
 	{
 		return false;
 	}
-	check(Outputs.Num() == 1);
+
+	if (Outputs.Num() == 0)
+	{
+		// Map expression has no output - thus invalid to use. This happens when, for example, referenced texture file is not found
+		UE_LOG(LogMDLImporter, Warning, TEXT("Failed to create valid expression for: %s from: %s"), *MapName,
+			*MapBakeParam.InputBakePath);
+		return false;
+	}
+
 
 	if (!FirstNormalExpression && MapNameLower.Find(TEXT("normal")) != INDEX_NONE)
 	{
diff --git a/Engine/Plugins/Enterprise/MDLImporter/Source/MDLImporter/Private/generator/ConstantExpressionFactory.cpp b/Engine/Plugins/Enterprise/MDLImporter/Source/MDLImporter/Private/generator/ConstantExpressionFactory.cpp
--- a/Engine/Plugins/Enterprise/MDLImporter/Source/MDLImporter/Private/generator/ConstantExpressionFactory.cpp
+++ b/Engine/Plugins/Enterprise/MDLImporter/Source/MDLImporter/Private/generator/ConstantExpressionFactory.cpp
@@ -145,16 +145,24 @@ namespace Generator
 				auto Handle = mi::base::make_handle(MDLConstant.get_interface<const mi::neuraylib::IValue_texture>());
 				const mi::base::Handle<const mi::neuraylib::ITexture> MDLTexture(Transaction.access<mi::neuraylib::ITexture>(Handle->get_value()));
 
-				Common::FTextureProperty Property;
-				Property.Path       = Mdl::Util::GetTextureFileName(MDLTexture.get());
-				float Gamma         = MDLTexture->get_effective_gamma();
-				Property.bIsSRGB    = Gamma != 1.0;
-				if (ProcessingNormapMap)
+				if (MDLTexture)
 				{
-					Property.CompressionSettings = TC_Normalmap;
+					Common::FTextureProperty Property;
+					Property.Path = Mdl::Util::GetTextureFileName(MDLTexture.get());
+					float Gamma = MDLTexture->get_effective_gamma();
+					Property.bIsSRGB = Gamma != 1.0;
+					if (ProcessingNormapMap)
+					{
+						Property.CompressionSettings = TC_Normalmap;
+					}
+					UTexture2D* Texture = TextureFactory->CreateTexture(CurrentMaterial->GetOuter(), Property, CurrentMaterial->GetFlags());
+					return { NewMaterialExpressionTextureObject(CurrentMaterial, Texture) };
+				}
+				else
+				{
+					LogMessages.Emplace(MDLImporterLogging::EMessageSeverity::Error, TEXT("Couldn't create texture."));
+					return {};
 				}
-				UTexture2D* Texture = TextureFactory->CreateTexture(CurrentMaterial->GetOuter(), Property, CurrentMaterial->GetFlags());
-				return {NewMaterialExpressionTextureObject(CurrentMaterial, Texture)};
 			}
 			case mi::neuraylib::IValue::VK_BSDF_MEASUREMENT:
 			{
@@ -165,7 +173,7 @@ namespace Generator
 				return {NewMaterialExpressionConstant(CurrentMaterial, 0.0f, 0.0f, 0.0f)};
 			}
 		}
-		check(false);
+		ensure(false);
 		return {};
 	}
 
diff --git a/Engine/Plugins/Enterprise/MDLImporter/Source/MDLImporter/Private/generator/MaterialExpressionFactory.h b/Engine/Plugins/Enterprise/MDLImporter/Source/MDLImporter/Private/generator/MaterialExpressionFactory.h
--- a/Engine/Plugins/Enterprise/MDLImporter/Source/MDLImporter/Private/generator/MaterialExpressionFactory.h
+++ b/Engine/Plugins/Enterprise/MDLImporter/Source/MDLImporter/Private/generator/MaterialExpressionFactory.h
@@ -126,7 +126,7 @@ namespace Generator
 
 	inline void FMaterialExpressionFactory::SetCurrentNormal(UMaterialExpression* InNormalExpression, bool InIsGeometryExpression)
 	{
-		check(InNormalExpression);
+		ensure(InNormalExpression);
 		CurrentNormalExpression = InNormalExpression;
 		bIsGeometryExpression   = InIsGeometryExpression;
 	}
