Copy of 9566613 by kerim.borchaev from //UE4/Dev-Enterprise-Minimal:

MDLImporter: Compiles on Linux, Mac/Linux shared libs

- shared libs are placed in Binaries\ThirdParty\MDL (same as for Win)

diff --git a/Engine/Plugins/Enterprise/MDLImporter/MDLImporter.uplugin b/Engine/Plugins/Enterprise/MDLImporter/MDLImporter.uplugin
--- a/Engine/Plugins/Enterprise/MDLImporter/MDLImporter.uplugin
+++ b/Engine/Plugins/Enterprise/MDLImporter/MDLImporter.uplugin
@@ -23,7 +23,9 @@
 			"Type": "Editor",
 			"LoadingPhase": "Default",
 			"WhitelistPlatforms": [
-				"Win64"
+				"Win64",
+				"Mac",
+				"Linux"
 			]
 		}
 	]
diff --git a/Engine/Plugins/Enterprise/MDLImporter/Source/MDLImporter/MDLImporter.Build.cs b/Engine/Plugins/Enterprise/MDLImporter/Source/MDLImporter/MDLImporter.Build.cs
--- a/Engine/Plugins/Enterprise/MDLImporter/Source/MDLImporter/MDLImporter.Build.cs
+++ b/Engine/Plugins/Enterprise/MDLImporter/Source/MDLImporter/MDLImporter.Build.cs
@@ -66,36 +66,17 @@ namespace UnrealBuildTool.Rules
 
                 PrivateDefinitions.Add("USE_MDLSDK");
 
-                string TargetPlatform = "Windows.x64";
-                string TargetExtension = "lib";
                 if ((Target.Platform == UnrealTargetPlatform.Win64) || (Target.Platform == UnrealTargetPlatform.Win32))
                 {
                     PublicDefinitions.Add("MI_PLATFORM_WINDOWS");
                 }
                 else if (Target.Platform == UnrealTargetPlatform.Linux)
                 {
-                    TargetPlatform = "Linux.x64";
-                    TargetExtension = "so";
+                    PublicDefinitions.Add("MI_PLATFORM_LINUX");
                 }
                 else if (Target.Platform == UnrealTargetPlatform.Mac)
                 {
-                    TargetPlatform = "MacOSX.x64";
-                    TargetExtension = "dylib";
-                }
-
-                // add static libraries
-                for (int i = 0; i < Libs.Length; ++i)
-                {
-                    string LibName = StaticLibNames[i];
-                    if (LibName == "")
-                    {
-                        continue;
-                    }
-
-                    LibName += "." + TargetExtension;
-
-                    string LibPath = Path.Combine(ThirdPartyPath, Libs[i], TargetPlatform, "lib", LibName);
-                    PublicAdditionalLibraries.Add(LibPath);
+                    PublicDefinitions.Add("MI_PLATFORM_MACOSX");
                 }
             }
         }
diff --git a/Engine/Plugins/Enterprise/MDLImporter/Source/MDLImporter/Private/MDLMapHandler.cpp b/Engine/Plugins/Enterprise/MDLImporter/Source/MDLImporter/Private/MDLMapHandler.cpp
--- a/Engine/Plugins/Enterprise/MDLImporter/Source/MDLImporter/Private/MDLMapHandler.cpp
+++ b/Engine/Plugins/Enterprise/MDLImporter/Source/MDLImporter/Private/MDLMapHandler.cpp
@@ -29,6 +29,10 @@ FMDLMapHandler::FMDLMapHandler(const Mdl::FApiContext& MdlContext)
 	MaterialExpressionFactory.SetFunctionLoader(FunctionLoader.Get());
 }
 
+FMDLMapHandler::~FMDLMapHandler()
+{
+}
+
 void FMDLMapHandler::PreImport(const mi::neuraylib::IMaterial_definition& MDLMaterialDefinition,
                                const mi::neuraylib::ICompiled_material&   MDLMaterial,
                                mi::neuraylib::ITransaction&               MDLTransaction)
diff --git a/Engine/Plugins/Enterprise/MDLImporter/Source/MDLImporter/Private/MDLMapHandler.h b/Engine/Plugins/Enterprise/MDLImporter/Source/MDLImporter/Private/MDLMapHandler.h
--- a/Engine/Plugins/Enterprise/MDLImporter/Source/MDLImporter/Private/MDLMapHandler.h
+++ b/Engine/Plugins/Enterprise/MDLImporter/Source/MDLImporter/Private/MDLMapHandler.h
@@ -23,7 +23,7 @@ class FMDLMapHandler : public Mdl::IMapDistilHandler
 {
 public:
 	FMDLMapHandler(const Mdl::FApiContext& MdlContext);
-	virtual ~FMDLMapHandler() = default;
+	virtual ~FMDLMapHandler();
 
 	virtual void PreImport(const mi::neuraylib::IMaterial_definition& MDLMaterialDefinition,
 	                       const mi::neuraylib::ICompiled_material&   MDLMaterial,
diff --git a/Engine/Plugins/Enterprise/MDLImporter/Source/MDLImporter/Private/MDLMaterialSelector.cpp b/Engine/Plugins/Enterprise/MDLImporter/Source/MDLImporter/Private/MDLMaterialSelector.cpp
--- a/Engine/Plugins/Enterprise/MDLImporter/Source/MDLImporter/Private/MDLMaterialSelector.cpp
+++ b/Engine/Plugins/Enterprise/MDLImporter/Source/MDLImporter/Private/MDLMaterialSelector.cpp
@@ -91,7 +91,7 @@ namespace MDLImporterImpl
 			}
 		}
 
-		return Opacity < 1.f && IOR != 1.f || MdlMaterial.Absorption.WasValueBaked();
+		return (Opacity < 1.f && IOR != 1.f) || MdlMaterial.Absorption.WasValueBaked();
 	}
 }
 
diff --git a/Engine/Plugins/Enterprise/MDLImporter/Source/MDLImporter/Private/generator/FunctionLoader.cpp b/Engine/Plugins/Enterprise/MDLImporter/Source/MDLImporter/Private/generator/FunctionLoader.cpp
--- a/Engine/Plugins/Enterprise/MDLImporter/Source/MDLImporter/Private/generator/FunctionLoader.cpp
+++ b/Engine/Plugins/Enterprise/MDLImporter/Source/MDLImporter/Private/generator/FunctionLoader.cpp
@@ -2,7 +2,7 @@
 
 #include "FunctionLoader.h"
 
-#include "FunctionGenerator.H"
+#include "FunctionGenerator.h"
 #include "common/Utility.h"
 
 #include "AssetRegistryModule.h"
diff --git a/Engine/Plugins/Enterprise/MDLImporter/Source/MDLImporter/Private/mdl/MdlSdkDefines.h b/Engine/Plugins/Enterprise/MDLImporter/Source/MDLImporter/Private/mdl/MdlSdkDefines.h
--- a/Engine/Plugins/Enterprise/MDLImporter/Source/MDLImporter/Private/mdl/MdlSdkDefines.h
+++ b/Engine/Plugins/Enterprise/MDLImporter/Source/MDLImporter/Private/mdl/MdlSdkDefines.h
@@ -2,6 +2,8 @@
 
 #pragma once
 
+#if PLATFORM_WINDOWS
+
 #ifndef MDLSDK_INCLUDES_START
 #define MDLSDK_INCLUDES_START \
 	THIRD_PARTY_INCLUDES_START \
@@ -14,3 +16,18 @@
 	__pragma(warning(pop)) \
 	THIRD_PARTY_INCLUDES_END
 #endif // MDLSDK_INCLUDE_END
+
+#else
+
+#ifndef MDLSDK_INCLUDES_START
+#define MDLSDK_INCLUDES_START \
+	THIRD_PARTY_INCLUDES_START 
+#endif // MDLSDK_INCLUDES_START
+
+#ifndef MDLSDK_INCLUDES_END
+#define MDLSDK_INCLUDES_END 
+	THIRD_PARTY_INCLUDES_END
+#endif // MDLSDK_INCLUDE_END
+
+#endif
+
