Fix CIS Warnings from step [Compile UnrealHeaderTool Mac]
Unused variable in MDLImporter.Build.cs

diff --git a/Engine/Plugins/Enterprise/MDLImporter/Source/MDLImporter/MDLImporter.Build.cs b/Engine/Plugins/Enterprise/MDLImporter/Source/MDLImporter/MDLImporter.Build.cs
--- a/Engine/Plugins/Enterprise/MDLImporter/Source/MDLImporter/MDLImporter.Build.cs
+++ b/Engine/Plugins/Enterprise/MDLImporter/Source/MDLImporter/MDLImporter.Build.cs
@@ -47,10 +47,7 @@ namespace UnrealBuildTool.Rules
             if (Directory.Exists(ThirdPartyPath))
             {
                 //third party libraries
-
                 string[] Libs = { "mdl-sdk-314800.830"};
-                string[] StaticLibNames = { "" };
-
                 foreach (string Lib in Libs)
                 {
                     string IncludePath = Path.Combine(ThirdPartyPath, Lib, "include");
