Merging CL# 9617448 from Dev-Enterprise to Release-4.24 to fix include tool warnings
#rnx

diff --git a/Engine/Plugins/Enterprise/MDLImporter/Source/MDLImporter/Private/mdl/MaterialDistiller.h b/Engine/Plugins/Enterprise/MDLImporter/Source/MDLImporter/Private/mdl/MaterialDistiller.h
--- a/Engine/Plugins/Enterprise/MDLImporter/Source/MDLImporter/Private/mdl/MaterialDistiller.h
+++ b/Engine/Plugins/Enterprise/MDLImporter/Source/MDLImporter/Private/mdl/MaterialDistiller.h
@@ -7,6 +7,13 @@
 #include "Containers/UnrealString.h"
 #include "Math/IntPoint.h"
 
+#ifdef USE_MDLSDK
+
+#include "mdl/Common.h"
+#include "mi/base/handle.h"
+
+#endif // #ifdef USE_MDLSDK
+
 namespace Mdl
 {
 	class FApiContext;
@@ -60,10 +67,6 @@ namespace Mdl
 
 #ifdef USE_MDLSDK
 
-#include "mdl/Common.h"
-
-#include "mi/base/handle.h"
-
 namespace mi
 {
 	namespace neuraylib
