MDL importer whileliste on linux & mac

diff --git a/Engine/Plugins/Enterprise/MDLImporter/MDLImporter.uplugin b/Engine/Plugins/Enterprise/MDLImporter/MDLImporter.uplugin
--- a/Engine/Plugins/Enterprise/MDLImporter/MDLImporter.uplugin
+++ b/Engine/Plugins/Enterprise/MDLImporter/MDLImporter.uplugin
@@ -15,7 +15,9 @@
 	"IsBetaVersion": true,
 	"Installed": false,
 	"SupportedTargetPlatforms": [
-		"Win64"
+		"Win64",
+		"Mac",
+		"Linux"
 	],
 	"Modules": [
 		{
