MDL/AxF Importers: made sure dlls are packaged with the engine

diff --git a/Engine/Plugins/Enterprise/AxFImporter/Source/AxFImporter/AxFImporter.Build.cs b/Engine/Plugins/Enterprise/AxFImporter/Source/AxFImporter/AxFImporter.Build.cs
--- a/Engine/Plugins/Enterprise/AxFImporter/Source/AxFImporter/AxFImporter.Build.cs
+++ b/Engine/Plugins/Enterprise/AxFImporter/Source/AxFImporter/AxFImporter.Build.cs
@@ -44,7 +44,19 @@ namespace UnrealBuildTool.Rules
                 }
             );
 
-            PublicDelayLoadDLLs.Add("AxFDecoding.dll");
+            if (Target.Platform == UnrealTargetPlatform.Win64)
+            {
+                PublicDelayLoadDLLs.Add("AxFDecoding.dll");
+				string ModulePath = Path.Combine(EngineDirectory, "Plugins/Enterprise/AxFImporter/Binaries/ThirdParty/AxF", Target.Platform.ToString(), "AxFDecoding.dll");
+				if (!File.Exists(ModulePath))
+                {
+					string Err = string.Format("AxF Decoding dll '{0}' not found.", ModulePath);
+					System.Console.WriteLine(Err);
+					throw new BuildException(Err);
+                }
+                RuntimeDependencies.Add(ModulePath);
+            }
+
 
             if (Directory.Exists(ThirdPartyPath))
             {
diff --git a/Engine/Plugins/Enterprise/MDLImporter/Source/MDLImporter/MDLImporter.Build.cs b/Engine/Plugins/Enterprise/MDLImporter/Source/MDLImporter/MDLImporter.Build.cs
--- a/Engine/Plugins/Enterprise/MDLImporter/Source/MDLImporter/MDLImporter.Build.cs
+++ b/Engine/Plugins/Enterprise/MDLImporter/Source/MDLImporter/MDLImporter.Build.cs
@@ -44,6 +44,28 @@ namespace UnrealBuildTool.Rules
                 }
             );
 
+            string MdlSdkModuleName = "libmdl_sdk.so";
+            
+            if (Target.Platform == UnrealTargetPlatform.Win64)
+            {
+                MdlSdkModuleName = "libmdl_sdk.dll";
+
+            }
+            else
+            {
+				MdlSdkModuleName = "libmdl_sdk.so";
+            }
+
+            PublicDelayLoadDLLs.Add(MdlSdkModuleName);
+            string ModulePath = Path.Combine(EngineDirectory, "Plugins/Enterprise/MDLImporter/Binaries/ThirdParty/MDL", Target.Platform.ToString(), MdlSdkModuleName);
+            if (!File.Exists(ModulePath))
+            {
+                string Err = string.Format("MDL SDK module '{0}' not found.", ModulePath);
+                System.Console.WriteLine(Err);
+                throw new BuildException(Err);
+            }
+            RuntimeDependencies.Add(ModulePath);
+
             if (Directory.Exists(ThirdPartyPath))
             {
                 //third party libraries
