Fixed Mac project generation error introduced in cl 10123140. Moved the Mac binaries from MacOSX to Mac.

#rnx

diff --git a/Engine/Plugins/Enterprise/MDLImporter/Source/MDLImporter/Private/MDLImporter.cpp b/Engine/Plugins/Enterprise/MDLImporter/Source/MDLImporter/Private/MDLImporter.cpp
--- a/Engine/Plugins/Enterprise/MDLImporter/Source/MDLImporter/Private/MDLImporter.cpp
+++ b/Engine/Plugins/Enterprise/MDLImporter/Source/MDLImporter/Private/MDLImporter.cpp
@@ -101,7 +101,7 @@ FMDLImporter::FMDLImporter(const FString& PluginPath)
 #if PLATFORM_WINDOWS
 	const FString Platform = TEXT("Win64");
 #elif PLATFORM_MAC
-	const FString Platform = TEXT("MacOSX");
+	const FString Platform = TEXT("Mac");
 #elif PLATFORM_LINUX
 	const FString Platform = TEXT("Linux");
 #else
