MDLImporter: added extra missing dll

diff --git a/Engine/Plugins/Enterprise/MDLImporter/Source/MDLImporter/MDLImporter.Build.cs b/Engine/Plugins/Enterprise/MDLImporter/Source/MDLImporter/MDLImporter.Build.cs
--- a/Engine/Plugins/Enterprise/MDLImporter/Source/MDLImporter/MDLImporter.Build.cs
+++ b/Engine/Plugins/Enterprise/MDLImporter/Source/MDLImporter/MDLImporter.Build.cs
@@ -3,6 +3,7 @@
 using UnrealBuildTool;
 using System;
 using System.IO;
+using System.Collections.Generic;
 
 namespace UnrealBuildTool.Rules
 {
@@ -44,27 +45,32 @@ namespace UnrealBuildTool.Rules
                 }
             );
 
-            string MdlSdkModuleName = "libmdl_sdk.so";
-            
+            List<string> RuntimeModuleNames = new List<string>();
+
             if (Target.Platform == UnrealTargetPlatform.Win64)
             {
-                MdlSdkModuleName = "libmdl_sdk.dll";
-
+				RuntimeModuleNames.Add("libmdl_sdk.dll");
+				RuntimeModuleNames.Add("nv_freeimage.dll");
             }
             else
             {
-				MdlSdkModuleName = "libmdl_sdk.so";
+                RuntimeModuleNames.Add("libmdl_sdk.so");
+                RuntimeModuleNames.Add("nv_freeimage.so");
+                RuntimeModuleNames.Add("dds.so");
             }
 
-            PublicDelayLoadDLLs.Add(MdlSdkModuleName);
-            string ModulePath = Path.Combine(EngineDirectory, "Plugins/Enterprise/MDLImporter/Binaries/ThirdParty/MDL", Target.Platform.ToString(), MdlSdkModuleName);
-            if (!File.Exists(ModulePath))
+            foreach (string RuntimeModuleName in RuntimeModuleNames)
             {
-                string Err = string.Format("MDL SDK module '{0}' not found.", ModulePath);
-                System.Console.WriteLine(Err);
-                throw new BuildException(Err);
+                PublicDelayLoadDLLs.Add(RuntimeModuleName);
+				string ModulePath = Path.Combine(EngineDirectory, "Plugins/Enterprise/MDLImporter/Binaries/ThirdParty/MDL", Target.Platform.ToString(), RuntimeModuleName);
+				if (!File.Exists(ModulePath))
+				{
+					string Err = string.Format("MDL SDK module '{0}' not found.", ModulePath);
+					System.Console.WriteLine(Err);
+					throw new BuildException(Err);
+				}
+	            RuntimeDependencies.Add(ModulePath);
             }
-            RuntimeDependencies.Add(ModulePath);
 
             if (Directory.Exists(ThirdPartyPath))
             {
