Fix for Mac editor build.

diff --git a/Engine/Plugins/Enterprise/MDLImporter/Source/MDLImporter/MDLImporter.Build.cs b/Engine/Plugins/Enterprise/MDLImporter/Source/MDLImporter/MDLImporter.Build.cs
--- a/Engine/Plugins/Enterprise/MDLImporter/Source/MDLImporter/MDLImporter.Build.cs
+++ b/Engine/Plugins/Enterprise/MDLImporter/Source/MDLImporter/MDLImporter.Build.cs
@@ -46,30 +46,46 @@ namespace UnrealBuildTool.Rules
             );
 
             List<string> RuntimeModuleNames = new List<string>();
+            string BinaryLibraryFolder = Path.Combine(EngineDirectory, "Plugins/Enterprise/MDLImporter/Binaries/ThirdParty/MDL", Target.Platform.ToString());
 
             if (Target.Platform == UnrealTargetPlatform.Win64)
             {
 				RuntimeModuleNames.Add("libmdl_sdk.dll");
 				RuntimeModuleNames.Add("nv_freeimage.dll");
+                    
+                foreach (string RuntimeModuleName in RuntimeModuleNames)
+                {
+                    string ModulePath = Path.Combine(BinaryLibraryFolder, RuntimeModuleName);
+                    if (!File.Exists(ModulePath))
+                    {
+                        string Err = string.Format("MDL SDK module '{0}' not found.", ModulePath);
+                        System.Console.WriteLine(Err);
+                        throw new BuildException(Err);
+                    }
+
+                    PublicDelayLoadDLLs.Add(RuntimeModuleName);
+                    RuntimeDependencies.Add(ModulePath);
+                }
             }
-            else
+            else if (Target.Platform == UnrealTargetPlatform.Mac)
             {
                 RuntimeModuleNames.Add("libmdl_sdk.so");
                 RuntimeModuleNames.Add("nv_freeimage.so");
                 RuntimeModuleNames.Add("dds.so");
-            }
 
-            foreach (string RuntimeModuleName in RuntimeModuleNames)
-            {
-                PublicDelayLoadDLLs.Add(RuntimeModuleName);
-				string ModulePath = Path.Combine(EngineDirectory, "Plugins/Enterprise/MDLImporter/Binaries/ThirdParty/MDL", Target.Platform.ToString(), RuntimeModuleName);
-				if (!File.Exists(ModulePath))
-				{
-					string Err = string.Format("MDL SDK module '{0}' not found.", ModulePath);
-					System.Console.WriteLine(Err);
-					throw new BuildException(Err);
-				}
-	            RuntimeDependencies.Add(ModulePath);
+                foreach (string RuntimeModuleName in RuntimeModuleNames)
+                {
+                    string ModulePath = Path.Combine(BinaryLibraryFolder, RuntimeModuleName);
+                    if (!File.Exists(ModulePath))
+                    {
+                        string Err = string.Format("MDL SDK module '{0}' not found.", ModulePath);
+                        System.Console.WriteLine(Err);
+                        throw new BuildException(Err);
+                    }
+
+                    PublicDelayLoadDLLs.Add(ModulePath);
+                    RuntimeDependencies.Add(ModulePath);
+                }
             }
 
             if (Directory.Exists(ThirdPartyPath))
