MDLImporter: Replaced check with if or ensure to deal with potential crash

diff --git a/Engine/Plugins/Enterprise/MDLImporter/Source/MDLImporter/Private/generator/GeneratorMaterialExpressions.cpp b/Engine/Plugins/Enterprise/MDLImporter/Source/MDLImporter/Private/generator/GeneratorMaterialExpressions.cpp
--- a/Engine/Plugins/Enterprise/MDLImporter/Source/MDLImporter/Private/generator/GeneratorMaterialExpressions.cpp
+++ b/Engine/Plugins/Enterprise/MDLImporter/Source/MDLImporter/Private/generator/GeneratorMaterialExpressions.cpp
@@ -2,6 +2,8 @@
 
 #include "MaterialExpressions.h"
 
+#include "common/Logging.h"
+
 namespace Generator
 {
 	bool   EvaluateBool(const FMaterialExpressionConnection& Connection);
@@ -397,6 +399,10 @@ namespace Generator
 		switch (Input.ConnectionType)
 		{
 			case Expression:
+				if (!ensure(Input.ExpressionData.Expression))
+				{
+					return 0;
+				}
 				if (Input.ExpressionData.Expression->IsA<UMaterialExpressionAbs>())
 				{
 					UMaterialExpressionAbs* Abs = Cast<UMaterialExpressionAbs>(Input.ExpressionData.Expression);
@@ -748,7 +754,7 @@ namespace Generator
 				return 4;
 			case Texture:
 			default:
-				check(false);
+				ensure(false);
 				return 0;
 		}
 	}
@@ -1083,7 +1089,7 @@ namespace Generator
 	                                               const FMaterialExpressionConnection& Less, const FMaterialExpressionConnection& Equal,
 	                                               const FMaterialExpressionConnection& Greater)
 	{
-		check(IsScalar(A) && IsScalar(B) &&
+		ensure(IsScalar(A) && IsScalar(B) &&
 		      (((ComponentCount(Less) == ComponentCount(Greater)) &&
 		        (!Equal.ExpressionData.Expression || (ComponentCount(Less) == ComponentCount(Equal))))));
 		UMaterialExpressionIf* Expression = NewMaterialExpression<UMaterialExpressionIf>(Parent);
@@ -1163,7 +1169,7 @@ namespace Generator
 		uint32 ACount     = ComponentCount(A);
 		uint32 BCount     = ComponentCount(B);
 		uint32 AlphaCount = ComponentCount(Alpha);
-		check(((ACount == 1) || (BCount == 1) || (ACount == BCount)) && ((AlphaCount == 1) || ((AlphaCount == ACount) && (AlphaCount == BCount))));
+		ensure(((ACount == 1) || (BCount == 1) || (ACount == BCount)) && ((AlphaCount == 1) || ((AlphaCount == ACount) && (AlphaCount == BCount))));
 
 		UMaterialExpressionLinearInterpolate* Expression = NewMaterialExpression<UMaterialExpressionLinearInterpolate>(Parent);
 		CheckedConnect(Parent, A, Expression->A, Expression->ConstA);
@@ -1207,7 +1213,7 @@ namespace Generator
 
 	UMaterialExpressionLogarithm2* NewMaterialExpressionLogarithm2(UObject* Parent, const FMaterialExpressionConnection& X)
 	{
-		check(IsScalar(X));
+		ensure(IsScalar(X));
 		UMaterialExpressionLogarithm2* Expression = NewMaterialExpression<UMaterialExpressionLogarithm2>(Parent);
 		CheckedConnect(Parent, X, Expression->X);
 		return Expression;
@@ -1221,16 +1227,21 @@ namespace Generator
 		UMaterialExpressionMaterialFunctionCall* Expression = NewMaterialExpression<UMaterialExpressionMaterialFunctionCall>(Parent);
 		Expression->SetMaterialFunction(Function);
 
-		checkf((Inputs.Num() == 0) || (Expression->FunctionInputs.Num() == Inputs.Num()),
-		       TEXT("Function <%s> expects %d inputs, but only got %d !\n"), *(Function->GetName()), Expression->FunctionInputs.Num(), Inputs.Num());
-
-		for (int32 i = 0; i < Inputs.Num(); i++)
+		if (Expression->FunctionInputs.Num() < Inputs.Num())
 		{
-			if ((Inputs[i].ConnectionType != EConnectionType::Expression) || !Inputs[i].ExpressionData.bIsDefault)
+			UE_LOG(LogMDLImporter, Warning, TEXT("Function <%s> received too many inputs - expected %d, but got %d!\n"), *(Function->GetName()), Expression->FunctionInputs.Num(), Inputs.Num());
+		}
+		else
+		{
+			for (int32 i = 0; i < Inputs.Num(); i++)
 			{
-				CheckedConnect(Parent, Inputs[i], Expression->FunctionInputs[i].Input);
+				if ((Inputs[i].ConnectionType != EConnectionType::Expression) || !Inputs[i].ExpressionData.bIsDefault)
+				{
+					CheckedConnect(Parent, Inputs[i], Expression->FunctionInputs[i].Input);
+				}
 			}
 		}
+		
 		Expression->UpdateFromFunctionResource();
 		return Expression;
 	}
