Merging CL 11386308 from //UE4/Dev-Enterprise to //UE4/Release-4.25

Fixed: ShaderCompileWorker terminated unexpectedly when importing MDL materials

diff --git a/Engine/Plugins/Enterprise/MDLImporter/Source/MDLImporter/Private/MDLMaterialPropertyFactory.cpp b/Engine/Plugins/Enterprise/MDLImporter/Source/MDLImporter/Private/MDLMaterialPropertyFactory.cpp
--- a/Engine/Plugins/Enterprise/MDLImporter/Source/MDLImporter/Private/MDLMaterialPropertyFactory.cpp
+++ b/Engine/Plugins/Enterprise/MDLImporter/Source/MDLImporter/Private/MDLMaterialPropertyFactory.cpp
@@ -370,9 +370,5 @@ Mat::FParameterMap FMDLMaterialPropertyFactory::CreateProperties(EObjectFlags Fl
 		}
 	}
 
-	Material.MarkPackageDirty();
-	Material.PostEditChange();
-	ParentPackage->MarkPackageDirty();
-
 	return ParameterMap;
 }
