MDL Importer: removed redundand foldoer creation and fixed two crashes('check' triggered on a material).

diff --git a/Engine/Plugins/Enterprise/MDLImporter/Source/MDLImporter/Private/MDLImporter.cpp b/Engine/Plugins/Enterprise/MDLImporter/Source/MDLImporter/Private/MDLImporter.cpp
--- a/Engine/Plugins/Enterprise/MDLImporter/Source/MDLImporter/Private/MDLImporter.cpp
+++ b/Engine/Plugins/Enterprise/MDLImporter/Source/MDLImporter/Private/MDLImporter.cpp
@@ -170,15 +170,6 @@ bool FMDLImporter::OpenFile(const FString& InFileName, const UMDLImporterOptions
 		return false;
 	}
 
-	// set export path for textures
-	{
-		const FString  ExporthPath  = FPaths::ConvertRelativePathToFull(FPaths::ProjectContentDir() + FPaths::GetBaseFilename(InFileName));
-		IPlatformFile& PlatformFile = FPlatformFileManager::Get().GetPlatformFile();
-		PlatformFile.CreateDirectory(*ExporthPath);
-
-		MdlContext->GetDistiller()->SetExportPath(ExporthPath);
-	}
-
 	bool bSuccess = false;
 	// load mdl module
 	{
diff --git a/Engine/Plugins/Enterprise/MDLImporter/Source/MDLImporter/Private/generator/GeneratorMaterialExpressions.cpp b/Engine/Plugins/Enterprise/MDLImporter/Source/MDLImporter/Private/generator/GeneratorMaterialExpressions.cpp
--- a/Engine/Plugins/Enterprise/MDLImporter/Source/MDLImporter/Private/generator/GeneratorMaterialExpressions.cpp
+++ b/Engine/Plugins/Enterprise/MDLImporter/Source/MDLImporter/Private/generator/GeneratorMaterialExpressions.cpp
@@ -741,7 +741,7 @@ namespace Generator
 					UMaterialExpressionDesaturation* Desaturate = Cast<UMaterialExpressionDesaturation>(Input.ExpressionData.Expression);
 					return ComponentCount(Desaturate->Input.Expression);
 				}
-				check(false);
+				ensure(false);
 				return 0;
 			case Boolean:
 			case Float:
@@ -1336,11 +1336,14 @@ namespace Generator
 	{
 		uint32 ACount = ComponentCount(A);
 		uint32 BCount = ComponentCount(B);
-		check((ACount == 1) || (BCount == 1) || (ACount == BCount));
-
 		UMaterialExpressionMultiply* Expression = NewMaterialExpression<UMaterialExpressionMultiply>(Parent);
-		CheckedConnect(Parent, A, Expression->A, Expression->ConstA);
-		CheckedConnect(Parent, B, Expression->B, Expression->ConstB);
+		
+		if (ensure((ACount == 1) || (BCount == 1) || (ACount == BCount)))
+		{
+			CheckedConnect(Parent, A, Expression->A, Expression->ConstA);
+			CheckedConnect(Parent, B, Expression->B, Expression->ConstB);
+		}
+
 		return Expression;
 	}
 
diff --git a/Engine/Plugins/Enterprise/MDLImporter/Source/MDLImporter/Private/mdl/MaterialDistiller.h b/Engine/Plugins/Enterprise/MDLImporter/Source/MDLImporter/Private/mdl/MaterialDistiller.h
--- a/Engine/Plugins/Enterprise/MDLImporter/Source/MDLImporter/Private/mdl/MaterialDistiller.h
+++ b/Engine/Plugins/Enterprise/MDLImporter/Source/MDLImporter/Private/mdl/MaterialDistiller.h
@@ -28,11 +28,6 @@ namespace Mdl
 
 		virtual ~IMaterialDistiller() = default;
 
-		/**
-		 * Sets the export path for the textures used by the materials to be loaded.
-		 */
-		virtual void SetExportPath(const FString& ExportPath) {}
-
 		/**
 		 * Sets a custom map handler used for the distillation.
 		 */
@@ -89,8 +84,6 @@ namespace Mdl
 
 		virtual ~FMaterialDistiller();
 
-		virtual void SetExportPath(const FString& ExportPath) override;
-
 		virtual void SetMapHanlder(IMapDistilHandler* MapHandler) override;
 
 		virtual void SetBakingSettings(uint32 Resolution, uint32 Samples) override;
@@ -150,7 +143,6 @@ namespace Mdl
 		IMapDistilHandler* MapHandler;
 		uint32             BakeResolution;
 		uint32             BakeSamples;
-		FString            BakePath;
 		float              MetersPerSceneUnit;
 
 		friend class FApiContext;
@@ -161,11 +153,6 @@ namespace Mdl
 		MapHandler = Handler;
 	}
 
-	inline void FMaterialDistiller::SetExportPath(const FString& Path)
-	{
-		BakePath = Path;
-	}
-
 	inline void FMaterialDistiller::SetBakingSettings(uint32 Resolution, uint32 Samples)
 	{
 		BakeResolution = FMath::Min<uint32>(ClosestPowerOfTwo(Resolution), 16384);
