Merging CL 15473104 from //UE4/Main

diff --git a/Engine/Plugins/Enterprise/MDLImporter/Source/MDLImporter/Public/MDLMaterialImporter.h b/Engine/Plugins/Enterprise/MDLImporter/Source/MDLImporter/Public/MDLMaterialImporter.h
--- a/Engine/Plugins/Enterprise/MDLImporter/Source/MDLImporter/Public/MDLMaterialImporter.h
+++ b/Engine/Plugins/Enterprise/MDLImporter/Source/MDLImporter/Public/MDLMaterialImporter.h
@@ -3,9 +3,11 @@
 #pragma once
 
 #include "CoreMinimal.h"
+#include "UObject/ObjectMacros.h"
 
 class UMaterialInterface;
 class UMDLImporterOptions;
+class UPackage;
 
 class MDLIMPORTER_API FMdlMaterialImporter
 {
