MDLImporter: produce warning instead of error when loading the plugin that was built without MDL SDK

diff --git a/Engine/Plugins/Enterprise/MDLImporter/Source/MDLImporter/Private/MDLImporterFactory.cpp b/Engine/Plugins/Enterprise/MDLImporter/Source/MDLImporter/Private/MDLImporterFactory.cpp
--- a/Engine/Plugins/Enterprise/MDLImporter/Source/MDLImporter/Private/MDLImporterFactory.cpp
+++ b/Engine/Plugins/Enterprise/MDLImporter/Source/MDLImporter/Private/MDLImporterFactory.cpp
@@ -118,7 +118,7 @@ UMDLImporterFactory::UMDLImporterFactory(const FObjectInitializer& ObjectInitial
 		Formats.Add(Extension);
 	}
 #else
-	UE_LOG(LogMDLImporter, Error, TEXT("MDL SDK was not available when plugin was build - MDL plugin won't be functional!"));
+	UE_LOG(LogMDLImporter, Warning, TEXT("MDL SDK was not available when plugin was build - MDL plugin won't be functional!"));
 #endif
 }
 
