Merging CL 15801364 (Fix UE4 References in MDLImporter and PlmXml translator)

//UE4/Main/...

to //UE5/Main/...

diff --git a/Engine/Plugins/Enterprise/MDLImporter/Source/MDLImporter/Private/MDLMapHandler.cpp b/Engine/Plugins/Enterprise/MDLImporter/Source/MDLImporter/Private/MDLMapHandler.cpp
--- a/Engine/Plugins/Enterprise/MDLImporter/Source/MDLImporter/Private/MDLMapHandler.cpp
+++ b/Engine/Plugins/Enterprise/MDLImporter/Source/MDLImporter/Private/MDLMapHandler.cpp
@@ -125,7 +125,7 @@ bool FMDLMapHandler::Import(const FString& MapName, bool bIsTexture, Mdl::FBakeP
 
 	if (MapNameLower.Find(TEXT("roughness")) != INDEX_NONE)
 	{
-		// convert from GGX to UE4
+		// convert from GGX to Unreal
 		UMaterialExpression* Roughness =
 		    Generator::NewMaterialExpressionSquareRoot(CurrentMaterial, {Outputs[0].GetExpressionAndUse(), Outputs[0].GetExpressionOutputIndex()});
 		MapBakeParam.SetExpression(Roughness, bIsTexture, 0);
diff --git a/Engine/Plugins/Enterprise/MDLImporter/Source/MDLImporter/Private/generator/FunctionGenerator.h b/Engine/Plugins/Enterprise/MDLImporter/Source/MDLImporter/Private/generator/FunctionGenerator.h
--- a/Engine/Plugins/Enterprise/MDLImporter/Source/MDLImporter/Private/generator/FunctionGenerator.h
+++ b/Engine/Plugins/Enterprise/MDLImporter/Source/MDLImporter/Private/generator/FunctionGenerator.h
@@ -256,7 +256,7 @@ namespace Generator
 			        NewMaterialExpressionIfEqual(Function, Anisotropy, 1.0f, 0.0001f, NewMaterialExpressionOneMinus(Function, Anisotropy))));
 
 			UMaterialExpressionMultiply* RotationAngle = NewMaterialExpressionMultiply(Function, 2.0f * PI, AnisotropyRotation);
-			// As UE4 is left-handed, swap StateNormal and TangentU in this cross product to get the TangentV !
+			// As Unreal is left-handed, swap StateNormal and TangentU in this cross product to get the TangentV !
 			UMaterialExpressionCrossProduct* TangentV =
 			    NewMaterialExpressionCrossProduct(Function, TangentU, NewMaterialExpressionFunctionCall(Function, StateNormal, {}));
 			UMaterialExpressionSubtract* RotatedTangentU =
@@ -508,7 +508,7 @@ namespace Generator
 			UMaterialExpressionMaterialFunctionCall* Normal = NewMaterialExpressionFunctionCall(
 			    Function, StateTransformVector,
 			    {(int)ECoordinateSpace::Internal, CoordinateSpace,
-			     StateNormalCall});  // would need to transform a normal here! but I'm not aware of such a method in UE4!
+			     StateNormalCall});  // would need to transform a normal here! but I'm not aware of such a method in Unreal!
 
 			UMaterialExpressionMaterialFunctionCall* Cubic = NewMaterialExpressionFunctionCall(
 			    Function, ImporterComputeCubicTransform,
@@ -1738,7 +1738,7 @@ namespace Generator
 			    NewMaterialExpressionStaticSwitch(Function, FlipTangentV, NewMaterialExpressionNegate(Function, UVWTangentV), UVWTangentV);
 #endif
 
-			// UE4 already applies bias to normal after sampling so no need to do it manually
+			// Unreal already applies bias to normal after sampling so no need to do it manually
 			UMaterialExpressionMultiply* TangentSpaceNormal = NewMaterialExpressionMultiply(
 			    Function,
 #if defined(USE_WORLD_ALIGNED_TEXTURES)
@@ -3295,7 +3295,7 @@ namespace Generator
 			check(0 == ArraySize);
 			Function->Description = TEXT("Array of geometry bitangents.");
 
-			// As UE4 is left handed, the v-tangent is (0, -1, 0) !
+			// As Unreal is left handed, the v-tangent is (0, -1, 0) !
 			NewMaterialExpressionFunctionOutput(Function, TEXT("geometry_tangent_v"), {0.0f, -1.0f, 0.0f});
 		}
 
@@ -3414,7 +3414,7 @@ namespace Generator
 			check(0 == ArraySize);
 			Function->Description = TEXT("The array of bitangent vectors for each texture space.");
 
-			// As UE4 is left handed, the v-tangent is (0, -1, 0) !
+			// As Unreal is left handed, the v-tangent is (0, -1, 0) !
 			NewMaterialExpressionFunctionOutput(Function, TEXT("texture_tangent_v"), {0.0f, -1.0f, 0.0f});
 		}
 
@@ -4289,7 +4289,7 @@ namespace Generator
 			    NewMaterialExpressionNormalize(Function, NewMaterialExpressionCrossProduct(Function, Normal, TangentUSimplePreliminary));
 			UMaterialExpressionIf* TangentVUDirCheck =
 			    NewMaterialExpressionIfGreater(Function, NormalDotUDir, 0.999f, TangentVComplex, TangentVSimple);
-			// As UE4 is left-handed, negate the resulting v-tangent !
+			// As Unreal is left-handed, negate the resulting v-tangent !
 			UMaterialExpressionSubtract* TangentV = NewMaterialExpressionNegate(
 			    Function, NewMaterialExpressionIfGreater(Function, NormalDotVDir, 0.999f, TangentVComplex, TangentVUDirCheck));
 
@@ -4385,7 +4385,7 @@ namespace Generator
 			                                NewMaterialExpressionNegate(Function, TangentVSimplePreliminary), TangentVSimplePreliminary);
 			UMaterialExpressionIf* TangentVUDirCheck =
 			    NewMaterialExpressionIfGreater(Function, NormalDotUDir, 0.999f, TangentVComplex, TangentVSimple);
-			// As UE4 is left-handed, negate the resulting v-tangent !
+			// As Unreal is left-handed, negate the resulting v-tangent !
 			UMaterialExpressionSubtract* TangentV = NewMaterialExpressionNegate(
 			    Function, NewMaterialExpressionIfGreater(Function, NormalDotVDir, 0.999f, TangentVComplex, TangentVUDirCheck));
 
@@ -5772,7 +5772,7 @@ namespace Generator
 			NewMaterialExpressionFunctionOutput(Function, TEXT("refl"), Refl);
 		}
 
-		void UE4TangentSpaceNormal(UMaterialFunction* Function, int32 ArraySize)
+		void UnrealTangentSpaceNormal(UMaterialFunction* Function, int32 ArraySize)
 		{
 			check(0 == ArraySize);
 			UMaterialFunction* StateNormal = LoadFunction(TEXT("mdl_state_normal"));
@@ -5802,7 +5802,7 @@ namespace Generator
 			NewMaterialExpressionFunctionOutput(Function, TEXT("normal"), UnclippedNormal);
 		}
 
-		void UE4OpacityWeight(UMaterialFunction* Function, int32 ArraySize)
+		void UnrealOpacityWeight(UMaterialFunction* Function, int32 ArraySize)
 		{
 			// always return 1
 			UMaterialExpressionFunctionInput* Opacity = NewMaterialExpressionFunctionInput(Function, TEXT("opacity"), EFunctionInputType::FunctionInput_Scalar, 1.0f);
diff --git a/Engine/Plugins/Enterprise/MDLImporter/Source/MDLImporter/Private/generator/FunctionLoader.cpp b/Engine/Plugins/Enterprise/MDLImporter/Source/MDLImporter/Private/generator/FunctionLoader.cpp
--- a/Engine/Plugins/Enterprise/MDLImporter/Source/MDLImporter/Private/generator/FunctionLoader.cpp
+++ b/Engine/Plugins/Enterprise/MDLImporter/Source/MDLImporter/Private/generator/FunctionLoader.cpp
@@ -222,9 +222,9 @@ namespace Generator
 		FunctionGenerateMap.Add(TEXT("mdl_nvidia_distilling_support_refl_from_ior_color"), {&FFunctionGenerator::DistillingSupportReflFromIORFloat3, 3});
 		FunctionGenerateMap.Add(TEXT("mdl_nvidia_distilling_support_refl_from_ior_float"), {&FFunctionGenerator::DistillingSupportReflFromIORFloat, 1});
 
-		// UE4 Omni
-		FunctionGenerateMap.Add(TEXT("mdl_OmniUe4Base_tangent_space_normal"), { &FFunctionGenerator::UE4TangentSpaceNormal, 1 });
-		FunctionGenerateMap.Add(TEXT("mdl_OmniUe4Translucent_get_translucent_opacity"), { &FFunctionGenerator::UE4OpacityWeight, 1 });
+		// Unreal Omni
+		FunctionGenerateMap.Add(TEXT("mdl_OmniUe4Base_tangent_space_normal"), { &FFunctionGenerator::UnrealTangentSpaceNormal, 1 });
+		FunctionGenerateMap.Add(TEXT("mdl_OmniUe4Translucent_get_translucent_opacity"), { &FFunctionGenerator::UnrealOpacityWeight, 1 });
 
 		// common functions generation
 		const FString EngineUtilityPath = TEXT("/Engine/Functions/Engine_MaterialFunctions02/Utility");
diff --git a/Engine/Plugins/Enterprise/MDLImporter/Source/MDLImporter/Private/generator/MaterialExpressionFactory.cpp b/Engine/Plugins/Enterprise/MDLImporter/Source/MDLImporter/Private/generator/MaterialExpressionFactory.cpp
--- a/Engine/Plugins/Enterprise/MDLImporter/Source/MDLImporter/Private/generator/MaterialExpressionFactory.cpp
+++ b/Engine/Plugins/Enterprise/MDLImporter/Source/MDLImporter/Private/generator/MaterialExpressionFactory.cpp
@@ -801,13 +801,13 @@ namespace Generator
 			}
 			case mi::neuraylib::IFunction_definition::DS_INTRINSIC_TEX_DEPTH:
 			{
-				// UE4 doesn't know about 3D textures ?? Does that mean, the depth is always 1 ?
+				// Unreal doesn't know about 3D textures ?? Does that mean, the depth is always 1 ?
 				check(Inputs.Num() == 1);
 				return {NewMaterialExpressionConstant(CurrentMaterial, 1.0f)};
 			}
 			case mi::neuraylib::IFunction_definition::DS_INTRINSIC_TEX_TEXTURE_ISVALID:
 			{
-				// in UE4, a texture is always valid, so just make a static bool <true>
+				// in Unreal, a texture is always valid, so just make a static bool <true>
 				return {NewMaterialExpressionStaticBool(CurrentMaterial, true)};
 			}
 			default:
diff --git a/Engine/Plugins/Enterprise/MDLImporter/Source/MDLImporter/Private/mdl/ApiContext.h b/Engine/Plugins/Enterprise/MDLImporter/Source/MDLImporter/Private/mdl/ApiContext.h
--- a/Engine/Plugins/Enterprise/MDLImporter/Source/MDLImporter/Private/mdl/ApiContext.h
+++ b/Engine/Plugins/Enterprise/MDLImporter/Source/MDLImporter/Private/mdl/ApiContext.h
@@ -68,7 +68,7 @@ namespace Mdl
 		}
 
 		/**
-		 * Returns the material distiller used for material distillation to the UE4 target model.
+		 * Returns the material distiller used for material distillation to the Unreal target model.
 		 */
 		virtual FMaterialDistiller* GetDistiller()
 		{
diff --git a/Engine/Plugins/Enterprise/MDLImporter/Source/MDLImporter/Private/mdl/MaterialDistiller.cpp b/Engine/Plugins/Enterprise/MDLImporter/Source/MDLImporter/Private/mdl/MaterialDistiller.cpp
--- a/Engine/Plugins/Enterprise/MDLImporter/Source/MDLImporter/Private/mdl/MaterialDistiller.cpp
+++ b/Engine/Plugins/Enterprise/MDLImporter/Source/MDLImporter/Private/mdl/MaterialDistiller.cpp
@@ -131,8 +131,8 @@ namespace Mdl
 			return GetParam((int)Type, MaterialBakeParam);
 		}
 
-		// Handles the material semantic for UE4 target models and sets the bake paths for the maps
-		void HandleMaterialSemanticUE4(mi::neuraylib::ITransaction*             Transaction,
+		// Handles the material semantic for Unreal target models and sets the bake paths for the maps
+		void HandleMaterialSemanticUnreal(mi::neuraylib::ITransaction*             Transaction,
 		                               const mi::neuraylib::ICompiled_material* CompiledMaterial,
 		                               TArray<FBakeParam>&                      BakeParams)
 		{
@@ -232,7 +232,7 @@ namespace Mdl
 		                       float MetersPerSceneUnit)
 		{
 			TArray<FBakeParam>& BakeParams = MaterialBakeParams;
-			// Setup some UE4 material parameters
+			// Setup some Unreal material parameters
 			BakeParams.Emplace((int)EParameterType::BaseColor, Material.BaseColor);
 			BakeParams.Emplace((int)EParameterType::Metallic, Material.Metallic);
 			BakeParams.Emplace((int)EParameterType::Roughness, Material.Roughness, RemapRoughness);
@@ -291,7 +291,7 @@ namespace Mdl
 
 			BakeParams.Emplace((int)EParameterType::IOR, Material.IOR);
 
-			HandleMaterialSemanticUE4(Transaction, DistilledMaterial, MaterialBakeParams);
+			HandleMaterialSemanticUnreal(Transaction, DistilledMaterial, MaterialBakeParams);
 
 			// Check for cutout-opacity
 
@@ -433,8 +433,8 @@ namespace Mdl
 	                                const FString&               MaterialCompiledDbName,
 	                                FMaterial&                   Material) const
 	{
-		// Always using UE4 material model
-		static const char* TargetModel = "ue4";
+		// Always using Unreal material model
+		static const char* TargetModel = "ue4"; // Name of the target model passed to the MDL SDK
 
 		TArray<FBakeParam> MaterialBakeParams;
 		MaterialBakeParams.Reserve((int)EParameterType::Count);
