Tessellation Removal: Removing Tessellation specific material parameters.

diff --git a/Engine/Plugins/Editor/GLTFImporter/Source/GLTFImporter/Private/GLTFMaterialElement.cpp b/Engine/Plugins/Editor/GLTFImporter/Source/GLTFImporter/Private/GLTFMaterialElement.cpp
--- a/Engine/Plugins/Editor/GLTFImporter/Source/GLTFImporter/Private/GLTFMaterialElement.cpp
+++ b/Engine/Plugins/Editor/GLTFImporter/Source/GLTFImporter/Private/GLTFMaterialElement.cpp
@@ -247,7 +247,6 @@ void FGLTFMaterialElement::Finalize()
 	ConnectInput(Opacity, MaterialExpressions, Material->Opacity);
 	ConnectInput(Refraction, MaterialExpressions, Material->Refraction);
 	ConnectInput(Normal, MaterialExpressions, Material->Normal);
-	ConnectInput(WorldDisplacement, MaterialExpressions, Material->WorldDisplacement);
 	ConnectInput(AmbientOcclusion, MaterialExpressions, Material->AmbientOcclusion);
 
 	UMaterialEditingLibrary::LayoutMaterialExpressions(Material);
diff --git a/Engine/Plugins/Enterprise/AxFImporter/Source/AxFImporter/Private/AxFImporter.cpp b/Engine/Plugins/Enterprise/AxFImporter/Source/AxFImporter/Private/AxFImporter.cpp
--- a/Engine/Plugins/Enterprise/AxFImporter/Source/AxFImporter/Private/AxFImporter.cpp
+++ b/Engine/Plugins/Enterprise/AxFImporter/Source/AxFImporter/Private/AxFImporter.cpp
@@ -1446,23 +1446,6 @@ private:
 					Connect(Material->OpacityMask, DitherTemporalAA);
 				}
 			}
-
-			if (ExpressionHeight)
-			{
-				UMaterialExpressionVertexNormalWS* VertexNormalWS = CreateMaterialExpression<UMaterialExpressionVertexNormalWS>();
-
-				UMaterialExpressionScalarParameter* DisplacementScale = CreateMaterialExpression<UMaterialExpressionScalarParameter>();
-				DisplacementScale->DefaultValue = 1.0f;
-				DisplacementScale->ParameterName = FName(TEXT("Displacement Scale"));
-				
-				Connect(Material->WorldDisplacement, Mul(VertexNormalWS, Mul(ExpressionHeight, DisplacementScale)));
-
-				UMaterialExpressionScalarParameter* TessellationMultiplier = CreateMaterialExpression<UMaterialExpressionScalarParameter>();
-				TessellationMultiplier->DefaultValue = 1.0f;
-				TessellationMultiplier->ParameterName = FName(TEXT("Tessellation Multiplier"));
-
-				Connect(Material->TessellationMultiplier, TessellationMultiplier);
-			}
 		}
 
 		void SetTextureDiffuseColor(FProcessedTextureSource Source) override
diff --git a/Engine/Plugins/Enterprise/DatasmithImporter/Source/DatasmithImporter/Private/DatasmithMaterialExpressions.cpp b/Engine/Plugins/Enterprise/DatasmithImporter/Source/DatasmithImporter/Private/DatasmithMaterialExpressions.cpp
--- a/Engine/Plugins/Enterprise/DatasmithImporter/Source/DatasmithImporter/Private/DatasmithMaterialExpressions.cpp
+++ b/Engine/Plugins/Enterprise/DatasmithImporter/Source/DatasmithImporter/Private/DatasmithMaterialExpressions.cpp
@@ -373,8 +373,6 @@ void FDatasmithMaterialExpressions::GetTextureSamplersFunc(UMaterialFunction* Ma
 	ExpressionInputs.Add(&Attrib->OpacityMask);
 	ExpressionInputs.Add(&Attrib->Normal);
 	ExpressionInputs.Add(&Attrib->WorldPositionOffset);
-	ExpressionInputs.Add(&Attrib->WorldDisplacement);
-	ExpressionInputs.Add(&Attrib->TessellationMultiplier);
 	ExpressionInputs.Add(&Attrib->SubsurfaceColor);
 	ExpressionInputs.Add(&Attrib->ClearCoat);
 	ExpressionInputs.Add(&Attrib->ClearCoatRoughness);
@@ -923,8 +921,6 @@ EMaterialProperty FDatasmithMaterialExpressions::DatasmithTextureSlotToMaterialP
 	case EDatasmithTextureSlot::OPACITYMASK:			return MP_OpacityMask;
 	case EDatasmithTextureSlot::NORMAL:					return MP_Normal;
 	case EDatasmithTextureSlot::WORLDPOSITIONOFFSET:	return MP_WorldPositionOffset;
-	case EDatasmithTextureSlot::DISPLACE:				return MP_WorldDisplacement;
-	case EDatasmithTextureSlot::TESSELLATIONMULTIPLIER:	return MP_TessellationMultiplier;
 	case EDatasmithTextureSlot::SUBSURFACECOLOR:		return MP_SubsurfaceColor;
 	case EDatasmithTextureSlot::COATSPECULAR:			return MP_CustomData0;
 	case EDatasmithTextureSlot::COATROUGHNESS:			return MP_CustomData1;
@@ -987,12 +983,6 @@ FExpressionInput* FDatasmithMaterialExpressions::GetMaterialOrFunctionSlot( UObj
 		case EDatasmithTextureSlot::WORLDPOSITIONOFFSET:
 			ExpressionInput = &Attrib->WorldPositionOffset;
 			break;
-		case EDatasmithTextureSlot::DISPLACE:
-			ExpressionInput = &Attrib->WorldDisplacement;
-			break;
-		case EDatasmithTextureSlot::TESSELLATIONMULTIPLIER:
-			ExpressionInput = &Attrib->TessellationMultiplier;
-			break;
 		case EDatasmithTextureSlot::SUBSURFACECOLOR:
 			ExpressionInput = &Attrib->SubsurfaceColor;
 			break;
@@ -2124,68 +2114,6 @@ void FDatasmithMaterialExpressions::CreateParallaxOffset(UObject* UnrealMatOrFun
 	}
 }
 
-void FDatasmithMaterialExpressions::ModulateDisplacement(UObject* UnrealMatOrFunc, double Amount, int32 SubDivisions)
-{
-	UMaterial* UnrealMaterial = Cast< UMaterial >( UnrealMatOrFunc );
-
-	UMaterialFunction* Func = Cast< UMaterialFunction >( UnrealMatOrFunc );
-
-	UMaterialExpressionMakeMaterialAttributes* Attrib = nullptr;
-	if (Func)
-	{
-		Attrib = FindOrAddAttributesFromMatFunc(Func);
-	}
-
-	UMaterialExpression* ToBeConnected = nullptr;
-	if (UnrealMaterial)
-	{
-		ToBeConnected = UnrealMaterial->WorldDisplacement.Expression;
-	}
-
-	if (Attrib)
-	{
-		ToBeConnected = Attrib->WorldDisplacement.Expression;
-	}
-
-	if (ToBeConnected == nullptr)
-	{
-		return;
-	}
-
-	UMaterialExpression* VertexWsExpr = CreateMaterialExpression<UMaterialExpressionVertexNormalWS>(UnrealMatOrFunc);
-
-	UMaterialExpression* SubDivExpr = Constant(UnrealMatOrFunc, SubDivisions);
-
-	if (UnrealMaterial)
-	{
-		UnrealMaterial->TessellationMultiplier.Expression = SubDivExpr;
-	}
-
-	if (Attrib)
-	{
-		Attrib->TessellationMultiplier.Expression = SubDivExpr;
-	}
-
-	UMaterialExpression* MultiplyWs = Multiply(UnrealMatOrFunc, ToBeConnected, 0.0, VertexWsExpr, 0.0);
-
-	if (ToBeConnected->Outputs.Num() > 1) // use red in case of colored Texture
-	{
-		ToBeConnected->ConnectExpression(MultiplyWs->GetInput(0), 1);
-	}
-
-	UMaterialExpression* multiplyamount = Multiply(UnrealMatOrFunc, MultiplyWs, 0.0, nullptr, Amount);
-
-	if (UnrealMaterial)
-	{
-		UnrealMaterial->WorldDisplacement.Expression = multiplyamount;
-	}
-
-	if (Attrib)
-	{
-		Attrib->WorldDisplacement.Expression = multiplyamount;
-	}
-}
-
 void FDatasmithMaterialExpressions::CreateDatasmithMaterialHelper(UPackage* Package, const TSharedPtr< IDatasmithShaderElement >& ShaderElement, const FDatasmithAssetsImportContext& AssetsContext,
 																UObject* UnrealMaterial)
 {
@@ -2513,8 +2441,6 @@ void FDatasmithMaterialExpressions::CreateDatasmithMaterialHelper(UPackage* Pack
 			AddTextureExpression(DisplaceTexture, ShaderElement->GetDisplaceTextureSampler(), 0.0, 0.0, UnrealMaterial, EDatasmithTextureSlot::DISPLACE);
 		}
 	}
-
-	ModulateDisplacement(UnrealMaterial, ShaderElement->GetDisplace(), ShaderElement->GetDisplaceSubDivision());
 }
 
 UMaterialExpressionMaterialFunctionCall* FDatasmithMaterialExpressions::BlendFunctions(UMaterial* UnrealMaterial, const FDatasmithAssetsImportContext& AssetsContext,
diff --git a/Engine/Plugins/Enterprise/DatasmithImporter/Source/DatasmithImporter/Public/DatasmithMaterialExpressions.h b/Engine/Plugins/Enterprise/DatasmithImporter/Source/DatasmithImporter/Public/DatasmithMaterialExpressions.h
--- a/Engine/Plugins/Enterprise/DatasmithImporter/Source/DatasmithImporter/Public/DatasmithMaterialExpressions.h
+++ b/Engine/Plugins/Enterprise/DatasmithImporter/Source/DatasmithImporter/Public/DatasmithMaterialExpressions.h
@@ -145,7 +145,6 @@ private:
 	static void CreateDatasmithMaterialCoat(const FDatasmithAssetsImportContext& AssetsContext, const TSharedPtr< IDatasmithShaderElement >& ShaderElement, UObject* UnrealMaterial);
 	static void ModulateNormalAmount(UObject* UnrealMatOrFunc, double Amount);
 	static void CreateParallaxOffset(UObject* UnrealMatOrFunc, UMaterialExpressionTextureObject* TextureObject, FDatasmithTextureSampler UV, float Amount);
-	static void ModulateDisplacement(UObject* UnrealMatOrFunc, double Amount, int32 SubDivisions);
 	static UMaterialExpressionMaterialFunctionCall* BlendFunctions(UMaterial* UnrealMaterial, const FDatasmithAssetsImportContext& AssetsContext, const TSharedPtr< IDatasmithShaderElement >& ShaderTop,
 																   UMaterialExpressionMaterialFunctionCall* CallTop, UMaterialExpressionMaterialFunctionCall* CallBase, UMaterialFunction* BlendFunc);
 
diff --git a/Engine/Plugins/Enterprise/MDLImporter/Source/MDLImporter/Private/generator/FunctionGenerator.h b/Engine/Plugins/Enterprise/MDLImporter/Source/MDLImporter/Private/generator/FunctionGenerator.h
--- a/Engine/Plugins/Enterprise/MDLImporter/Source/MDLImporter/Private/generator/FunctionGenerator.h
+++ b/Engine/Plugins/Enterprise/MDLImporter/Source/MDLImporter/Private/generator/FunctionGenerator.h
@@ -2316,7 +2316,7 @@ namespace Generator
 			    Function, TEXT("normal"), EFunctionInputType::FunctionInput_Vector3, NewMaterialExpressionFunctionCall(Function, StateNormal, {}));
 
 			UMaterialExpressionMakeMaterialAttributes* BSDF =
-			    NewMaterialExpressionMakeMaterialAttributes(Function, Tint, {}, 0.0f, {}, {}, {}, {}, Normal, {}, {}, {}, Tint);
+			    NewMaterialExpressionMakeMaterialAttributes(Function, Tint, {}, 0.0f, {}, {}, {}, {}, Normal, {}, Tint);
 
 			NewMaterialExpressionFunctionOutput(Function, TEXT("bsdf"), BSDF);
 		}
@@ -3695,22 +3695,20 @@ namespace Generator
 			                                                {BlendBreak, 6},     // OpacityMask
 			                                                {LayerBreak, 7},     // Normal purely from the layer
 			                                                {BlendBreak, 8},     // WorldPositionOffset,
-			                                                {BlendBreak, 9},     // WorldDisplacement
-			                                                {BlendBreak, 10},    // TessellationMultiplier
-			                                                {BlendBreak, 11},    // SubsurfaceColor,
+			                                                {BlendBreak, 9},    // SubsurfaceColor,
 			                                                ClearCoat,           // ClearCoat
 			                                                ClearCoatRoughness,  // ClearCoatRoughness
-			                                                {BlendBreak, 14},    // AmbientOcclusion,
-			                                                {BlendBreak, 15},    // Refraction
-			                                                {BlendBreak, 16},    // CustomizedUVs0
-			                                                {BlendBreak, 17},    // CustomizedUVs1
-			                                                {BlendBreak, 18},    // CustomizedUVs2
-			                                                {BlendBreak, 19},    // CustomizedUVs3
-			                                                {BlendBreak, 20},    // CustomizedUVs4
-			                                                {BlendBreak, 21},    // CustomizedUVs5
-			                                                {BlendBreak, 22},    // CustomizedUVs6
-			                                                {BlendBreak, 23},    // CustomizedUVs7
-			                                                {BlendBreak, 24});   // PixelDepthOffset
+			                                                {BlendBreak, 12},    // AmbientOcclusion,
+			                                                {BlendBreak, 13},    // Refraction
+			                                                {BlendBreak, 14},    // CustomizedUVs0
+			                                                {BlendBreak, 15},    // CustomizedUVs1
+			                                                {BlendBreak, 16},    // CustomizedUVs2
+			                                                {BlendBreak, 17},    // CustomizedUVs3
+			                                                {BlendBreak, 18},    // CustomizedUVs4
+			                                                {BlendBreak, 19},    // CustomizedUVs5
+			                                                {BlendBreak, 20},    // CustomizedUVs6
+			                                                {BlendBreak, 21},    // CustomizedUVs7
+			                                                {BlendBreak, 22});   // PixelDepthOffset
 
 			NewMaterialExpressionFunctionOutput(Function, TEXT("bsdf"), BSDF);
 		}
@@ -4873,9 +4871,7 @@ namespace Generator
 			                                                NewMaterialExpressionIfEqual(Function, Value, 0.0f, {BreakFalse, 19}, {BreakTrue, 19}),
 			                                                NewMaterialExpressionIfEqual(Function, Value, 0.0f, {BreakFalse, 20}, {BreakTrue, 20}),
 			                                                NewMaterialExpressionIfEqual(Function, Value, 0.0f, {BreakFalse, 21}, {BreakTrue, 21}),
-			                                                NewMaterialExpressionIfEqual(Function, Value, 0.0f, {BreakFalse, 22}, {BreakTrue, 22}),
-			                                                NewMaterialExpressionIfEqual(Function, Value, 0.0f, {BreakFalse, 23}, {BreakTrue, 23}),
-			                                                NewMaterialExpressionIfEqual(Function, Value, 0.0f, {BreakFalse, 24}, {BreakTrue, 24}));
+			                                                NewMaterialExpressionIfEqual(Function, Value, 0.0f, {BreakFalse, 22}, {BreakTrue, 22}));
 
 			NewMaterialExpressionFunctionOutput(Function, TEXT("bsdf"), BSDF);
 		}
@@ -4904,22 +4900,20 @@ namespace Generator
 			                                                                                              OpacityMask,   // OpacityMask
 			                                                                                              {Break, 7},    // Normal
 			                                                                                              {Break, 8},    // WorldPositionOffset
-			                                                                                              {Break, 9},    // WorldDisplacement
-			                                                                                              {Break, 10},   // TessellationMultiplier
-			                                                                                              {Break, 11},   // SubsurfaceColor
-			                                                                                              {Break, 12},   // ClearCoat
-			                                                                                              {Break, 13},   // ClearCoatRoughness
-			                                                                                              {Break, 14},   // AmbientOcclusion
-			                                                                                              {Break, 15},   // Refraction
-			                                                                                              {Break, 16},   // CustomizedUVs0
-			                                                                                              {Break, 17},   // CustomizedUVs1
-			                                                                                              {Break, 18},   // CustomizedUVs2
-			                                                                                              {Break, 19},   // CustomizedUVs3
-			                                                                                              {Break, 20},   // CustomizedUVs4
-			                                                                                              {Break, 21},   // CustomizedUVs5
-			                                                                                              {Break, 22},   // CustomizedUVs6
-			                                                                                              {Break, 23},   // CustomizedUVs7
-			                                                                                              {Break, 24});  // PixelDepthOffset
+			                                                                                              {Break, 9},   // SubsurfaceColor
+			                                                                                              {Break, 10},   // ClearCoat
+			                                                                                              {Break, 11},   // ClearCoatRoughness
+			                                                                                              {Break, 12},   // AmbientOcclusion
+			                                                                                              {Break, 13},   // Refraction
+			                                                                                              {Break, 14},   // CustomizedUVs0
+			                                                                                              {Break, 15},   // CustomizedUVs1
+			                                                                                              {Break, 16},   // CustomizedUVs2
+			                                                                                              {Break, 17},   // CustomizedUVs3
+			                                                                                              {Break, 18},   // CustomizedUVs4
+			                                                                                              {Break, 19},   // CustomizedUVs5
+			                                                                                              {Break, 20},   // CustomizedUVs6
+			                                                                                              {Break, 21},   // CustomizedUVs7
+			                                                                                              {Break, 22});  // PixelDepthOffset
 
 			NewMaterialExpressionFunctionOutput(Function, TEXT("bsdf"), BSDF);
 		}
@@ -4948,22 +4942,20 @@ namespace Generator
 			                                                {MaterialBreak, 6},    // OpacityMask
 			                                                {MaterialBreak, 7},    // Normal
 			                                                {MaterialBreak, 8},    // WorldPositionOffset,
-			                                                {MaterialBreak, 9},    // WorldDisplacement
-			                                                {MaterialBreak, 10},   // TessellationMultiplier
-			                                                {MaterialBreak, 11},   // SubsurfaceColor,
-			                                                {MaterialBreak, 12},   // ClearCoat
-			                                                {MaterialBreak, 13},   // ClearCoatRoughness
-			                                                {MaterialBreak, 14},   // AmbientOcclusion,
+			                                                {MaterialBreak, 9},   // SubsurfaceColor,
+			                                                {MaterialBreak, 10},   // ClearCoat
+			                                                {MaterialBreak, 11},   // ClearCoatRoughness
+			                                                {MaterialBreak, 12},   // AmbientOcclusion,
 			                                                Refraction,            // Refraction
-			                                                {MaterialBreak, 16},   // CustomizedUVs0
-			                                                {MaterialBreak, 17},   // CustomizedUVs1
-			                                                {MaterialBreak, 18},   // CustomizedUVs2
-			                                                {MaterialBreak, 19},   // CustomizedUVs3
-			                                                {MaterialBreak, 20},   // CustomizedUVs4
-			                                                {MaterialBreak, 21},   // CustomizedUVs5
-			                                                {MaterialBreak, 22},   // CustomizedUVs6
-			                                                {MaterialBreak, 23},   // CustomizedUVs7
-			                                                {MaterialBreak, 24});  // PixelDepthOffset
+			                                                {MaterialBreak, 14},   // CustomizedUVs0
+			                                                {MaterialBreak, 15},   // CustomizedUVs1
+			                                                {MaterialBreak, 16},   // CustomizedUVs2
+			                                                {MaterialBreak, 17},   // CustomizedUVs3
+			                                                {MaterialBreak, 18},   // CustomizedUVs4
+			                                                {MaterialBreak, 19},   // CustomizedUVs5
+			                                                {MaterialBreak, 20},   // CustomizedUVs6
+			                                                {MaterialBreak, 21},   // CustomizedUVs7
+			                                                {MaterialBreak, 22});  // PixelDepthOffset
 
 			NewMaterialExpressionFunctionOutput(Function, TEXT("bsdf"), BSDF);
 		}
@@ -4990,22 +4982,20 @@ namespace Generator
 			                                                {MaterialBreak, 6},    // OpacityMask
 			                                                {MaterialBreak, 7},    // Normal
 			                                                {MaterialBreak, 8},    // WorldPositionOffset,
-			                                                {MaterialBreak, 9},    // WorldDisplacement
-			                                                {MaterialBreak, 10},   // TessellationMultiplier
 			                                                SubsurfaceColor,       // SubsurfaceColor,
-			                                                {MaterialBreak, 12},   // ClearCoat
-			                                                {MaterialBreak, 13},   // ClearCoatRoughness
-			                                                {MaterialBreak, 14},   // AmbientOcclusion,
-			                                                {MaterialBreak, 15},   // Refraction
-			                                                {MaterialBreak, 16},   // CustomizedUVs0
-			                                                {MaterialBreak, 17},   // CustomizedUVs1
-			                                                {MaterialBreak, 18},   // CustomizedUVs2
-			                                                {MaterialBreak, 19},   // CustomizedUVs3
-			                                                {MaterialBreak, 20},   // CustomizedUVs4
-			                                                {MaterialBreak, 21},   // CustomizedUVs5
-			                                                {MaterialBreak, 22},   // CustomizedUVs6
-			                                                {MaterialBreak, 23},   // CustomizedUVs7
-			                                                {MaterialBreak, 24});  // PixelDepthOffset
+			                                                {MaterialBreak, 10},   // ClearCoat
+			                                                {MaterialBreak, 11},   // ClearCoatRoughness
+			                                                {MaterialBreak, 12},   // AmbientOcclusion,
+			                                                {MaterialBreak, 13},   // Refraction
+			                                                {MaterialBreak, 14},   // CustomizedUVs0
+			                                                {MaterialBreak, 15},   // CustomizedUVs1
+			                                                {MaterialBreak, 16},   // CustomizedUVs2
+			                                                {MaterialBreak, 17},   // CustomizedUVs3
+			                                                {MaterialBreak, 18},   // CustomizedUVs4
+			                                                {MaterialBreak, 19},   // CustomizedUVs5
+			                                                {MaterialBreak, 20},   // CustomizedUVs6
+			                                                {MaterialBreak, 21},   // CustomizedUVs7
+			                                                {MaterialBreak, 22});  // PixelDepthOffset
 
 			NewMaterialExpressionFunctionOutput(Function, TEXT("bsdf"), BSDF);
 		}
diff --git a/Engine/Plugins/Enterprise/MDLImporter/Source/MDLImporter/Private/generator/GeneratorMaterialExpressions.cpp b/Engine/Plugins/Enterprise/MDLImporter/Source/MDLImporter/Private/generator/GeneratorMaterialExpressions.cpp
--- a/Engine/Plugins/Enterprise/MDLImporter/Source/MDLImporter/Private/generator/GeneratorMaterialExpressions.cpp
+++ b/Engine/Plugins/Enterprise/MDLImporter/Source/MDLImporter/Private/generator/GeneratorMaterialExpressions.cpp
@@ -1265,8 +1265,7 @@ namespace Generator
 	    const FMaterialExpressionConnection& Specular, const FMaterialExpressionConnection& Roughness,
 	    const FMaterialExpressionConnection& EmissiveColor, const FMaterialExpressionConnection& Opacity,
 	    const FMaterialExpressionConnection& OpacityMask, const FMaterialExpressionConnection& Normal,
-	    const FMaterialExpressionConnection& WorldPositionOffset, const FMaterialExpressionConnection& WorldDisplacement,
-	    const FMaterialExpressionConnection& TessellationMultiplier, const FMaterialExpressionConnection& SubsurfaceColor,
+	    const FMaterialExpressionConnection& WorldPositionOffset, const FMaterialExpressionConnection& SubsurfaceColor,
 	    const FMaterialExpressionConnection& ClearCoat, const FMaterialExpressionConnection& ClearCoatRoughness,
 	    const FMaterialExpressionConnection& AmbientOcclusion, const FMaterialExpressionConnection& Refraction,
 	    const FMaterialExpressionConnection& CustomizedUVs0, const FMaterialExpressionConnection& CustomizedUVs1,
@@ -1276,7 +1275,7 @@ namespace Generator
 	    const FMaterialExpressionConnection& PixelDepthOffset)
 	{
 		check(IsScalar(Metallic) && IsScalar(Specular) && IsScalar(Roughness) && IsScalar(Opacity) && IsScalar(OpacityMask) &&
-		      IsScalar(TessellationMultiplier) && IsScalar(ClearCoat) && IsScalar(ClearCoatRoughness) && IsScalar(AmbientOcclusion) &&
+		      IsScalar(ClearCoat) && IsScalar(ClearCoatRoughness) && IsScalar(AmbientOcclusion) &&
 		      IsScalar(PixelDepthOffset));
 
 		UMaterialExpressionMakeMaterialAttributes* Expression = NewMaterialExpression<UMaterialExpressionMakeMaterialAttributes>(Parent);
@@ -1290,8 +1289,6 @@ namespace Generator
 		CheckedConnect(Parent, OpacityMask, Expression->OpacityMask);
 		CheckedConnect(Parent, Normal, Expression->Normal);
 		CheckedConnect(Parent, WorldPositionOffset, Expression->WorldPositionOffset);
-		CheckedConnect(Parent, WorldDisplacement, Expression->WorldDisplacement);
-		CheckedConnect(Parent, TessellationMultiplier, Expression->TessellationMultiplier);
 		CheckedConnect(Parent, SubsurfaceColor, Expression->SubsurfaceColor);
 		CheckedConnect(Parent, ClearCoat, Expression->ClearCoat);
 		CheckedConnect(Parent, ClearCoatRoughness, Expression->ClearCoatRoughness);
diff --git a/Engine/Plugins/Enterprise/MDLImporter/Source/MDLImporter/Private/generator/MaterialExpressions.h b/Engine/Plugins/Enterprise/MDLImporter/Source/MDLImporter/Private/generator/MaterialExpressions.h
--- a/Engine/Plugins/Enterprise/MDLImporter/Source/MDLImporter/Private/generator/MaterialExpressions.h
+++ b/Engine/Plugins/Enterprise/MDLImporter/Source/MDLImporter/Private/generator/MaterialExpressions.h
@@ -257,8 +257,7 @@ namespace Generator
 	    const FMaterialExpressionConnection& Specular = {}, const FMaterialExpressionConnection& Roughness = {},
 	    const FMaterialExpressionConnection& EmissiveColor = {}, const FMaterialExpressionConnection& Opacity = {},
 	    const FMaterialExpressionConnection& OpacityMask = {}, const FMaterialExpressionConnection& Normal = {},
-	    const FMaterialExpressionConnection& WorldPositionOffset = {}, const FMaterialExpressionConnection& WorldDisplacement = {},
-	    const FMaterialExpressionConnection& TessellationMultiplier = {}, const FMaterialExpressionConnection& SubsurfaceColor = {},
+	    const FMaterialExpressionConnection& WorldPositionOffset = {}, const FMaterialExpressionConnection& SubsurfaceColor = {},
 	    const FMaterialExpressionConnection& ClearCoat = 0.0f, const FMaterialExpressionConnection& ClearCoatRoughness = {},
 	    const FMaterialExpressionConnection& AmbientOcclusion = {}, const FMaterialExpressionConnection& Refraction = {},
 	    const FMaterialExpressionConnection& CustomizedUVs0 = {}, const FMaterialExpressionConnection& CustomizedUVs1 = {},
diff --git a/Engine/Plugins/Enterprise/MDLImporter/Source/MDLImporter/Private/material/BakedMaterialFactory.cpp b/Engine/Plugins/Enterprise/MDLImporter/Source/MDLImporter/Private/material/BakedMaterialFactory.cpp
--- a/Engine/Plugins/Enterprise/MDLImporter/Source/MDLImporter/Private/material/BakedMaterialFactory.cpp
+++ b/Engine/Plugins/Enterprise/MDLImporter/Source/MDLImporter/Private/material/BakedMaterialFactory.cpp
@@ -139,27 +139,6 @@ namespace Mat
 			MapConnecter.DeleteExpression(EMaterialParameter::EmissionStrength);
 		}
 
-		// displacement
-		if (Parameters.Contains(EMaterialParameter::DisplacementMap))
-		{
-			UMaterialExpression* UV = Generator::NewMaterialExpressionTextureCoordinate(&Material, 0);
-			UV                      = Generator::NewMaterialExpressionMultiply(&Material, {UV, Tiling});
-			UMaterialExpression* Displacement =
-			    Generator::NewMaterialExpressionTextureSample(&Material, {Parameters[EMaterialParameter::DisplacementMap]}, UV);
-			// taking third(B/Z) output for Displacement - mdl distiller bakes displacement multiplied by state::normal so that scalar value goes to Z component
-			Displacement = Generator::NewMaterialExpressionMultiply(&Material, { {Displacement, 3}, Parameters[EMaterialParameter::DisplacementStrength] });
-			UMaterialExpression* WorldNormal = Generator::NewMaterialExpression<UMaterialExpressionVertexNormalWS>(&Material);
-
-			Generator::Connect(Material.WorldDisplacement, Generator::NewMaterialExpressionMultiply(&Material, {WorldNormal, Displacement}));
-
-			UMaterialExpression* Multiplier = Generator::NewMaterialExpressionScalarParameter(&Material, TEXT("Tesselation Multiplier"), 1.f);
-			Generator::Connect(Material.TessellationMultiplier, Multiplier);
-
-			Generator::SetMaterialExpressionGroup(TEXT("Displacement"), Multiplier);
-			Generator::SetMaterialExpressionGroup(TEXT("Displacement"), Parameters[EMaterialParameter::DisplacementMap]);
-			Generator::SetMaterialExpressionGroup(TEXT("Displacement"), Parameters[EMaterialParameter::DisplacementStrength]);
-		}
-
 		// never used
 		MapConnecter.DeleteExpressionMap(EMaterialParameter::Opacity);
 		MapConnecter.DeleteExpression(EMaterialParameter::IOR);
diff --git a/Engine/Shaders/Private/MaterialTemplate.ush b/Engine/Shaders/Private/MaterialTemplate.ush
--- a/Engine/Shaders/Private/MaterialTemplate.ush
+++ b/Engine/Shaders/Private/MaterialTemplate.ush
@@ -2405,20 +2405,6 @@ float3 GetMaterialPreviousWorldPositionOffset(FMaterialVertexParameters Paramete
 %s;
 }
 
-half3 GetMaterialWorldDisplacement(FMaterialTessellationParameters Parameters)
-{
-%s;
-}
-
-half GetMaterialMaxDisplacement()
-{
-	return 0.0f;
-}
-
-half GetMaterialTessellationMultiplier(FMaterialTessellationParameters Parameters)
-{
-%s;
-}
 
 // .rgb:SubsurfaceColor, .a:SSProfileId in 0..1 range
 half4 GetMaterialSubsurfaceDataRaw(FPixelMaterialInputs PixelMaterialInputs)
diff --git a/Engine/Source/Developer/MeshMergeUtilities/Private/MeshMergeUtilities.cpp b/Engine/Source/Developer/MeshMergeUtilities/Private/MeshMergeUtilities.cpp
--- a/Engine/Source/Developer/MeshMergeUtilities/Private/MeshMergeUtilities.cpp
+++ b/Engine/Source/Developer/MeshMergeUtilities/Private/MeshMergeUtilities.cpp
@@ -598,8 +598,6 @@ EFlattenMaterialProperties FMeshMergeUtilities::NewToOldProperty(int32 NewProper
 		EFlattenMaterialProperties::NumFlattenMaterialProperties,
 		EFlattenMaterialProperties::NumFlattenMaterialProperties,
 		EFlattenMaterialProperties::NumFlattenMaterialProperties,
-		EFlattenMaterialProperties::NumFlattenMaterialProperties,
-		EFlattenMaterialProperties::NumFlattenMaterialProperties,
 		EFlattenMaterialProperties::AmbientOcclusion
 	};
 	
diff --git a/Engine/Source/Editor/MaterialEditor/Private/MaterialEditor.cpp b/Engine/Source/Editor/MaterialEditor/Private/MaterialEditor.cpp
--- a/Engine/Source/Editor/MaterialEditor/Private/MaterialEditor.cpp
+++ b/Engine/Source/Editor/MaterialEditor/Private/MaterialEditor.cpp
@@ -4253,7 +4253,6 @@ UClass* FMaterialEditor::GetOnPromoteToParameterClass(const UEdGraphPin* TargetP
 			case MP_Specular:
 			case MP_Roughness:
 			case MP_Anisotropy:
-			case MP_TessellationMultiplier:
 			case MP_CustomData0:
 			case MP_CustomData1:
 			case MP_AmbientOcclusion:
@@ -4265,7 +4264,6 @@ UClass* FMaterialEditor::GetOnPromoteToParameterClass(const UEdGraphPin* TargetP
 				return UMaterialExpressionScalarParameter::StaticClass();
 
 			case MP_WorldPositionOffset:
-			case MP_WorldDisplacement:
 			case MP_EmissiveColor:
 			case MP_BaseColor:
 			case MP_SubsurfaceColor:
diff --git a/Engine/Source/Editor/UnrealEd/Private/MaterialGraph.cpp b/Engine/Source/Editor/UnrealEd/Private/MaterialGraph.cpp
--- a/Engine/Source/Editor/UnrealEd/Private/MaterialGraph.cpp
+++ b/Engine/Source/Editor/UnrealEd/Private/MaterialGraph.cpp
@@ -71,8 +71,6 @@ void UMaterialGraph::RebuildGraphInternal(const TMap<UMaterialExpression*, TArra
 		MaterialInputs.Add(FMaterialInputInfo(FMaterialAttributeDefinitionMap::GetDisplayNameForMaterial(MP_Normal, Material), MP_Normal, LOCTEXT("NormalToolTip", "Takes the input of a normal map")));
 		MaterialInputs.Add(FMaterialInputInfo(FMaterialAttributeDefinitionMap::GetDisplayNameForMaterial(MP_Tangent, Material), MP_Tangent, LOCTEXT("TangentToolTip", "Takes the input of a tangent map. Useful for specifying anisotropy direction.")));
 		MaterialInputs.Add(FMaterialInputInfo(FMaterialAttributeDefinitionMap::GetDisplayNameForMaterial(MP_WorldPositionOffset, Material), MP_WorldPositionOffset, LOCTEXT("WorldPositionOffsetToolTip", "Allows for the vertices of a mesh to be manipulated in world space by the Material")));
-		MaterialInputs.Add(FMaterialInputInfo(FMaterialAttributeDefinitionMap::GetDisplayNameForMaterial(MP_WorldDisplacement, Material), MP_WorldDisplacement, LOCTEXT("WorldDisplacementToolTip", "Allows for the tessellation vertices to be manipulated in world space by the Material")));
-		MaterialInputs.Add(FMaterialInputInfo(FMaterialAttributeDefinitionMap::GetDisplayNameForMaterial(MP_TessellationMultiplier, Material), MP_TessellationMultiplier, LOCTEXT("TessllationMultiplierToolTip", "Controls the amount tessellation along the surface")));
 		MaterialInputs.Add(FMaterialInputInfo(FMaterialAttributeDefinitionMap::GetDisplayNameForMaterial(MP_SubsurfaceColor, Material), MP_SubsurfaceColor, LOCTEXT("SubsurfaceToolTip", "Allows you to add a color to your Material to simulate shifts in color when light passes through the surface")));
 		MaterialInputs.Add(FMaterialInputInfo(FMaterialAttributeDefinitionMap::GetDisplayNameForMaterial(MP_CustomData0, Material), MP_CustomData0, FMaterialAttributeDefinitionMap::GetDisplayNameForMaterial(MP_CustomData0, Material)));
 		MaterialInputs.Add(FMaterialInputInfo(FMaterialAttributeDefinitionMap::GetDisplayNameForMaterial(MP_CustomData1, Material), MP_CustomData1, FMaterialAttributeDefinitionMap::GetDisplayNameForMaterial(MP_CustomData1, Material)));
diff --git a/Engine/Source/Runtime/Core/Public/UObject/UE5MainStreamObjectVersion.h b/Engine/Source/Runtime/Core/Public/UObject/UE5MainStreamObjectVersion.h
--- a/Engine/Source/Runtime/Core/Public/UObject/UE5MainStreamObjectVersion.h
+++ b/Engine/Source/Runtime/Core/Public/UObject/UE5MainStreamObjectVersion.h
@@ -106,6 +106,9 @@ struct CORE_API FUE5MainStreamObjectVersion
 		// Fix to the Cubic smoothing, plus introduction of new smoothing types
         BlendSpaceSmoothingImprovements,
 
+		// Removing Tessellation parameters from Materials
+		RemovingTessellationParameters,
+
         // -----<new versions can be added above this line>-------------------------------------------------
 		VersionPlusOne,
 		LatestVersion = VersionPlusOne - 1
diff --git a/Engine/Source/Runtime/Engine/Classes/Materials/Material.h b/Engine/Source/Runtime/Engine/Classes/Materials/Material.h
--- a/Engine/Source/Runtime/Engine/Classes/Materials/Material.h
+++ b/Engine/Source/Runtime/Engine/Classes/Materials/Material.h
@@ -434,14 +434,6 @@ public:
 	FVectorMaterialInput WorldPositionOffset;
 
 #if WITH_EDITORONLY_DATA
-	/** Offset in world space applied to tessellated vertices. */
-	UPROPERTY()
-	FVectorMaterialInput WorldDisplacement;
-
-	/** Multiplies the tessellation factors applied when a tessellation mode is set. */
-	UPROPERTY()
-	FScalarMaterialInput TessellationMultiplier;
-
 	/** Inner material color, only used for ShadingModel=Subsurface */
 	UPROPERTY()
 	FColorMaterialInput SubsurfaceColor;
diff --git a/Engine/Source/Runtime/Engine/Classes/Materials/MaterialExpressionMakeMaterialAttributes.h b/Engine/Source/Runtime/Engine/Classes/Materials/MaterialExpressionMakeMaterialAttributes.h
--- a/Engine/Source/Runtime/Engine/Classes/Materials/MaterialExpressionMakeMaterialAttributes.h
+++ b/Engine/Source/Runtime/Engine/Classes/Materials/MaterialExpressionMakeMaterialAttributes.h
@@ -47,12 +47,6 @@ class UMaterialExpressionMakeMaterialAttributes : public UMaterialExpression
 	UPROPERTY()
 	FExpressionInput WorldPositionOffset;
 
-	UPROPERTY()
-	FExpressionInput WorldDisplacement;
-
-	UPROPERTY()
-	FExpressionInput TessellationMultiplier;
-
 	UPROPERTY()
 	FExpressionInput SubsurfaceColor;
 
diff --git a/Engine/Source/Runtime/Engine/Private/Materials/HLSLMaterialTranslator.cpp b/Engine/Source/Runtime/Engine/Private/Materials/HLSLMaterialTranslator.cpp
--- a/Engine/Source/Runtime/Engine/Private/Materials/HLSLMaterialTranslator.cpp
+++ b/Engine/Source/Runtime/Engine/Private/Materials/HLSLMaterialTranslator.cpp
@@ -715,8 +715,6 @@ bool FHLSLMaterialTranslator::Translate()
 			Chunk[MP_OpacityMask]					= Material->CompilePropertyAndSetMaterialProperty(MP_OpacityMask			,this);
 			Chunk[MP_Tangent]						= Material->CompilePropertyAndSetMaterialProperty(MP_Tangent				,this);
 			Chunk[MP_WorldPositionOffset]			= Material->CompilePropertyAndSetMaterialProperty(MP_WorldPositionOffset	,this);
-			Chunk[MP_WorldDisplacement]				= Material->CompilePropertyAndSetMaterialProperty(MP_WorldDisplacement		,this);
-			Chunk[MP_TessellationMultiplier]		= Material->CompilePropertyAndSetMaterialProperty(MP_TessellationMultiplier	,this);			
 
 			// Make sure to compile this property before using ShadingModelsFromCompilation
 			Chunk[MP_ShadingModel]					= Material->CompilePropertyAndSetMaterialProperty(MP_ShadingModel			,this);
@@ -1971,8 +1969,6 @@ FString FHLSLMaterialTranslator::GetMaterialShaderCode()
 
 	LazyPrintf.PushParam(!bEnableExecutionFlow ? *GenerateFunctionCode(MP_WorldPositionOffset, BaseDerivativeVariation) : TEXT("return Parameters.MaterialVertexAttributes.WorldPositionOffset"));
 	LazyPrintf.PushParam(!bEnableExecutionFlow ? *GenerateFunctionCode(CompiledMP_PrevWorldPositionOffset, BaseDerivativeVariation) : TEXT("return 0.0f"));
-	LazyPrintf.PushParam(!bEnableExecutionFlow ? *GenerateFunctionCode(MP_WorldDisplacement, BaseDerivativeVariation) : TEXT("return 0.0f"));
-	LazyPrintf.PushParam(!bEnableExecutionFlow ? *GenerateFunctionCode(MP_TessellationMultiplier, BaseDerivativeVariation) : TEXT("return 0.0f"));
 	LazyPrintf.PushParam(!bEnableExecutionFlow ? *GenerateFunctionCode(MP_CustomData0, BaseDerivativeVariation) : TEXT("return 0.0f"));
 	LazyPrintf.PushParam(!bEnableExecutionFlow ? *GenerateFunctionCode(MP_CustomData1, BaseDerivativeVariation) : TEXT("return 0.0f"));
 
diff --git a/Engine/Source/Runtime/Engine/Private/Materials/Material.cpp b/Engine/Source/Runtime/Engine/Private/Materials/Material.cpp
--- a/Engine/Source/Runtime/Engine/Private/Materials/Material.cpp
+++ b/Engine/Source/Runtime/Engine/Private/Materials/Material.cpp
@@ -240,17 +240,6 @@ int32 FMaterialResource::CompilePropertyAndSetMaterialProperty(EMaterialProperty
 			}
 			break;
 
-		case MP_WorldDisplacement:
-			if (Compiler->GetFeatureLevel() >= ERHIFeatureLevel::SM5)
-			{
-				Ret = MaterialInterface->CompileProperty(Compiler, Property);
-			}
-			else
-			{
-				Ret = FMaterialAttributeDefinitionMap::CompileDefaultExpression(Compiler, Property);
-			}
-			break;
-
 		case MP_MaterialAttributes:
 			Ret = MaterialInterface->CompileProperty(Compiler, Property);
 			break;
@@ -3419,7 +3408,7 @@ void UMaterial::Serialize(FArchive& Ar)
 	}
 #endif // #if WITH_EDITOR
 
-	static_assert(MP_MAX == 33, "New material properties must have DoMaterialAttributeReorder called on them to ensure that any future reordering of property pins is correctly applied.");
+	static_assert(MP_MAX == 31, "New material properties must have DoMaterialAttributeReorder called on them to ensure that any future reordering of property pins is correctly applied.");
 
 	if (Ar.UEVer() < VER_UE4_MATERIAL_MASKED_BLENDMODE_TIDY)
 	{
@@ -3665,7 +3654,6 @@ void UMaterial::ConvertMaterialToStrataMaterial()
 
 		// STRATA_TODO for mateiral conversion
 		//  - WorldPositionOffset can remain on the end point node
-		//  - WorldDisplacement and TessellationMultiplier are going to be removed and should be ignored?
 		//  - ShadingModel
 		//  - Refraction
 		//  - PixelDepthOffset
@@ -4047,38 +4035,37 @@ void UMaterial::PostLoad()
 #if WITH_EDITORONLY_DATA
 	const int32 UEVer = GetLinkerUEVersion();
 	const int32 RenderObjVer = GetLinkerCustomVersion(FRenderingObjectVersion::GUID);
-
-	DoMaterialAttributeReorder(&DiffuseColor_DEPRECATED, UEVer, RenderObjVer);
-	DoMaterialAttributeReorder(&SpecularColor_DEPRECATED, UEVer, RenderObjVer);
-	DoMaterialAttributeReorder(&BaseColor, UEVer, RenderObjVer);
-	DoMaterialAttributeReorder(&Metallic, UEVer, RenderObjVer);
-	DoMaterialAttributeReorder(&Specular, UEVer, RenderObjVer);
-	DoMaterialAttributeReorder(&Roughness, UEVer, RenderObjVer);
-	DoMaterialAttributeReorder(&Anisotropy, UEVer, RenderObjVer);
-	DoMaterialAttributeReorder(&Normal, UEVer, RenderObjVer);
-	DoMaterialAttributeReorder(&Tangent, UEVer, RenderObjVer);
-	DoMaterialAttributeReorder(&EmissiveColor, UEVer, RenderObjVer);
-	DoMaterialAttributeReorder(&Opacity, UEVer, RenderObjVer);
-	DoMaterialAttributeReorder(&OpacityMask, UEVer, RenderObjVer);
-	DoMaterialAttributeReorder(&WorldPositionOffset, UEVer, RenderObjVer);
-	DoMaterialAttributeReorder(&WorldDisplacement, UEVer, RenderObjVer);
-	DoMaterialAttributeReorder(&TessellationMultiplier, UEVer, RenderObjVer);
-	DoMaterialAttributeReorder(&SubsurfaceColor, UEVer, RenderObjVer);
-	DoMaterialAttributeReorder(&ClearCoat, UEVer, RenderObjVer);
-	DoMaterialAttributeReorder(&ClearCoatRoughness, UEVer, RenderObjVer);
-	DoMaterialAttributeReorder(&AmbientOcclusion, UEVer, RenderObjVer);
-	DoMaterialAttributeReorder(&Refraction, UEVer, RenderObjVer);
-	DoMaterialAttributeReorder(&CustomizedUVs[0], UEVer, RenderObjVer);
-	DoMaterialAttributeReorder(&CustomizedUVs[1], UEVer, RenderObjVer);
-	DoMaterialAttributeReorder(&CustomizedUVs[2], UEVer, RenderObjVer);
-	DoMaterialAttributeReorder(&CustomizedUVs[3], UEVer, RenderObjVer);
-	DoMaterialAttributeReorder(&CustomizedUVs[4], UEVer, RenderObjVer);
-	DoMaterialAttributeReorder(&CustomizedUVs[5], UEVer, RenderObjVer);
-	DoMaterialAttributeReorder(&CustomizedUVs[6], UEVer, RenderObjVer);
-	DoMaterialAttributeReorder(&CustomizedUVs[7], UEVer, RenderObjVer);
-	DoMaterialAttributeReorder(&PixelDepthOffset, UEVer, RenderObjVer);
-	DoMaterialAttributeReorder(&ShadingModelFromMaterialExpression, UEVer, RenderObjVer);
-	DoMaterialAttributeReorder(&FrontMaterial, UEVer, RenderObjVer);
+	const int32 UE5MainVer = GetLinkerCustomVersion(FUE5MainStreamObjectVersion::GUID);
+
+	DoMaterialAttributeReorder(&DiffuseColor_DEPRECATED, UEVer, RenderObjVer, UE5MainVer);
+	DoMaterialAttributeReorder(&SpecularColor_DEPRECATED, UEVer, RenderObjVer, UE5MainVer);
+	DoMaterialAttributeReorder(&BaseColor, UEVer, RenderObjVer, UE5MainVer);
+	DoMaterialAttributeReorder(&Metallic, UEVer, RenderObjVer, UE5MainVer);
+	DoMaterialAttributeReorder(&Specular, UEVer, RenderObjVer, UE5MainVer);
+	DoMaterialAttributeReorder(&Roughness, UEVer, RenderObjVer, UE5MainVer);
+	DoMaterialAttributeReorder(&Anisotropy, UEVer, RenderObjVer, UE5MainVer);
+	DoMaterialAttributeReorder(&Normal, UEVer, RenderObjVer, UE5MainVer);
+	DoMaterialAttributeReorder(&Tangent, UEVer, RenderObjVer, UE5MainVer);
+	DoMaterialAttributeReorder(&EmissiveColor, UEVer, RenderObjVer, UE5MainVer);
+	DoMaterialAttributeReorder(&Opacity, UEVer, RenderObjVer, UE5MainVer);
+	DoMaterialAttributeReorder(&OpacityMask, UEVer, RenderObjVer, UE5MainVer);
+	DoMaterialAttributeReorder(&WorldPositionOffset, UEVer, RenderObjVer, UE5MainVer);
+	DoMaterialAttributeReorder(&SubsurfaceColor, UEVer, RenderObjVer, UE5MainVer);
+	DoMaterialAttributeReorder(&ClearCoat, UEVer, RenderObjVer, UE5MainVer);
+	DoMaterialAttributeReorder(&ClearCoatRoughness, UEVer, RenderObjVer, UE5MainVer);
+	DoMaterialAttributeReorder(&AmbientOcclusion, UEVer, RenderObjVer, UE5MainVer);
+	DoMaterialAttributeReorder(&Refraction, UEVer, RenderObjVer, UE5MainVer);
+	DoMaterialAttributeReorder(&CustomizedUVs[0], UEVer, RenderObjVer, UE5MainVer);
+	DoMaterialAttributeReorder(&CustomizedUVs[1], UEVer, RenderObjVer, UE5MainVer);
+	DoMaterialAttributeReorder(&CustomizedUVs[2], UEVer, RenderObjVer, UE5MainVer);
+	DoMaterialAttributeReorder(&CustomizedUVs[3], UEVer, RenderObjVer, UE5MainVer);
+	DoMaterialAttributeReorder(&CustomizedUVs[4], UEVer, RenderObjVer, UE5MainVer);
+	DoMaterialAttributeReorder(&CustomizedUVs[5], UEVer, RenderObjVer, UE5MainVer);
+	DoMaterialAttributeReorder(&CustomizedUVs[6], UEVer, RenderObjVer, UE5MainVer);
+	DoMaterialAttributeReorder(&CustomizedUVs[7], UEVer, RenderObjVer, UE5MainVer);
+	DoMaterialAttributeReorder(&PixelDepthOffset, UEVer, RenderObjVer, UE5MainVer);
+	DoMaterialAttributeReorder(&ShadingModelFromMaterialExpression, UEVer, RenderObjVer, UE5MainVer);
+	DoMaterialAttributeReorder(&FrontMaterial, UEVer, RenderObjVer, UE5MainVer);
 #endif // WITH_EDITORONLY_DATA
 
 	if (!IsDefaultMaterial())
@@ -5727,8 +5714,6 @@ bool UMaterial::GetExpressionInputDescription(EMaterialProperty InProperty, FMat
 	case MP_Normal: SetMaterialInputDescription(Normal, false, OutDescription); return true;
 	case MP_Tangent: SetMaterialInputDescription(Tangent, false, OutDescription); return true;
 	case MP_WorldPositionOffset: SetMaterialInputDescription(WorldPositionOffset, false, OutDescription); return true;
-	case MP_WorldDisplacement: SetMaterialInputDescription(WorldDisplacement, false, OutDescription); return true;
-	case MP_TessellationMultiplier: SetMaterialInputDescription(TessellationMultiplier, false, OutDescription); return true;
 	case MP_SubsurfaceColor: SetMaterialInputDescription(SubsurfaceColor, false, OutDescription); return true;
 	case MP_CustomData0: SetMaterialInputDescription(ClearCoat, false, OutDescription); return true;
 	case MP_CustomData1: SetMaterialInputDescription(ClearCoatRoughness, false, OutDescription); return true;
@@ -6300,7 +6285,6 @@ int32 UMaterial::CompilePropertyEx( FMaterialCompiler* Compiler, const FGuid& At
 		case MP_Specular:				return Specular.CompileWithDefault(Compiler, Property);
 		case MP_Roughness:				return Roughness.CompileWithDefault(Compiler, Property);
 		case MP_Anisotropy:				return Anisotropy.CompileWithDefault(Compiler, Property);
-		case MP_TessellationMultiplier:	return TessellationMultiplier.CompileWithDefault(Compiler, Property);
 		case MP_CustomData0:			return ClearCoat.CompileWithDefault(Compiler, Property);
 		case MP_CustomData1:			return ClearCoatRoughness.CompileWithDefault(Compiler, Property);
 		case MP_AmbientOcclusion:		return AmbientOcclusion.CompileWithDefault(Compiler, Property);
@@ -6311,7 +6295,6 @@ int32 UMaterial::CompilePropertyEx( FMaterialCompiler* Compiler, const FGuid& At
 		case MP_Normal:					return Normal.CompileWithDefault(Compiler, Property);
 		case MP_Tangent:				return Tangent.CompileWithDefault(Compiler, Property);
 		case MP_WorldPositionOffset:	return WorldPositionOffset.CompileWithDefault(Compiler, Property);
-		case MP_WorldDisplacement:		return WorldDisplacement.CompileWithDefault(Compiler, Property);
 		case MP_PixelDepthOffset:		return PixelDepthOffset.CompileWithDefault(Compiler, Property);
 		case MP_ShadingModel:			return ShadingModelFromMaterialExpression.CompileWithDefault(Compiler, Property);
 		case MP_FrontMaterial:			return FrontMaterial.CompileWithDefault(Compiler, Property);
@@ -6776,10 +6759,6 @@ static bool IsPropertyActive_Internal(EMaterialProperty InProperty,
 	case MP_CustomData1:
 		Active = ShadingModels.HasAnyShadingModel({ MSM_ClearCoat, MSM_Eye });
 		break;
-	case MP_TessellationMultiplier:
-	case MP_WorldDisplacement:
-		Active = false;
-		break;
 	case MP_EmissiveColor:
 		// Emissive is always active, even for light functions and post process materials, 
 		// but not for AlphaHoldout
diff --git a/Engine/Source/Runtime/Engine/Private/Materials/MaterialExpressions.cpp b/Engine/Source/Runtime/Engine/Private/Materials/MaterialExpressions.cpp
--- a/Engine/Source/Runtime/Engine/Private/Materials/MaterialExpressions.cpp
+++ b/Engine/Source/Runtime/Engine/Private/Materials/MaterialExpressions.cpp
@@ -9,6 +9,7 @@
 #include "Misc/Guid.h"
 #include "UObject/RenderingObjectVersion.h"
 #include "UObject/EditorObjectVersion.h"
+#include "UObject/UE5MainStreamObjectVersion.h"
 #include "Misc/App.h"
 #include "UObject/Object.h"
 #include "UObject/Class.h"
@@ -774,14 +775,21 @@ void UMaterialExpression::Serialize(FStructuredArchive::FRecord Record)
 	SCOPED_LOADTIMER(UMaterialExpression::Serialize);
 	Super::Serialize(Record);
 
-	Record.GetUnderlyingArchive().UsingCustomVersion(FRenderingObjectVersion::GUID);
+	FArchive& Archive = Record.GetUnderlyingArchive();
+
+	Archive.UsingCustomVersion(FRenderingObjectVersion::GUID);
+	Archive.UsingCustomVersion(FUE5MainStreamObjectVersion::GUID);
 
 #if WITH_EDITORONLY_DATA
+	const int32 UEVer = Archive.UEVer();
+	const int32 RenderVer = Archive.CustomVer(FRenderingObjectVersion::GUID);
+	const int32 UE5Ver = Archive.CustomVer(FUE5MainStreamObjectVersion::GUID);
+
 	const TArray<FExpressionInput*> Inputs = GetInputs();
 	for (int32 InputIndex = 0; InputIndex < Inputs.Num(); ++InputIndex)
 	{
 		FExpressionInput* Input = Inputs[InputIndex];
-		DoMaterialAttributeReorder(Input, Record.GetUnderlyingArchive().UEVer(), Record.GetUnderlyingArchive().CustomVer(FRenderingObjectVersion::GUID));
+		DoMaterialAttributeReorder(Input, UEVer, RenderVer, UE5Ver);
 	}
 #endif // WITH_EDITORONLY_DATA
 }
@@ -5720,7 +5728,7 @@ int32 UMaterialExpressionMakeMaterialAttributes::Compile(class FMaterialCompiler
 	int32 Ret = INDEX_NONE;
 	UMaterialExpression* Expression = nullptr;
 
- 	static_assert(MP_MAX == 33, 
+ 	static_assert(MP_MAX == 31, 
 		"New material properties should be added to the end of the inputs for this expression. \
 		The order of properties here should match the material results pins, the make material attriubtes node inputs and the mapping of IO indices to properties in GetMaterialPropertyFromInputOutputIndex().\
 		Insertions into the middle of the properties or a change in the order of properties will also require that existing data is fixed up in DoMaterialAttributeReorder().\
@@ -5741,8 +5749,6 @@ int32 UMaterialExpressionMakeMaterialAttributes::Compile(class FMaterialCompiler
 	case MP_Normal: Ret = Normal.Compile(Compiler); Expression = Normal.Expression; break;
 	case MP_Tangent: Ret = Tangent.Compile(Compiler); Expression = Tangent.Expression; break;
 	case MP_WorldPositionOffset: Ret = WorldPositionOffset.Compile(Compiler); Expression = WorldPositionOffset.Expression; break;
-	case MP_WorldDisplacement: Ret = WorldDisplacement.Compile(Compiler); Expression = WorldDisplacement.Expression; break;
-	case MP_TessellationMultiplier: Ret = TessellationMultiplier.Compile(Compiler); Expression = TessellationMultiplier.Expression; break;
 	case MP_SubsurfaceColor: Ret = SubsurfaceColor.Compile(Compiler); Expression = SubsurfaceColor.Expression; break;
 	case MP_CustomData0: Ret = ClearCoat.Compile(Compiler); Expression = ClearCoat.Expression; break;
 	case MP_CustomData1: Ret = ClearCoatRoughness.Compile(Compiler); Expression = ClearCoatRoughness.Expression; break;
@@ -5829,7 +5835,7 @@ UMaterialExpressionBreakMaterialAttributes::UMaterialExpressionBreakMaterialAttr
 
 	MenuCategories.Add(ConstructorStatics.NAME_MaterialAttributes);
 	
- 	static_assert(MP_MAX == 33, 
+ 	static_assert(MP_MAX == 31, 
 		"New material properties should be added to the end of the outputs for this expression. \
 		The order of properties here should match the material results pins, the make material attriubtes node inputs and the mapping of IO indices to properties in GetMaterialPropertyFromInputOutputIndex().\
 		Insertions into the middle of the properties or a change in the order of properties will also require that existing data is fixed up in DoMaterialAttriubtesReorder().\
@@ -5847,8 +5853,6 @@ UMaterialExpressionBreakMaterialAttributes::UMaterialExpressionBreakMaterialAttr
 	Outputs.Add(FExpressionOutput(TEXT("Normal"), 1, 1, 1, 1, 0));
 	Outputs.Add(FExpressionOutput(TEXT("Tangent"), 1, 1, 1, 1, 0));
 	Outputs.Add(FExpressionOutput(TEXT("WorldPositionOffset"), 1, 1, 1, 1, 0));
-	Outputs.Add(FExpressionOutput(TEXT("WorldDisplacement"), 1, 1, 1, 1, 0));
-	Outputs.Add(FExpressionOutput(TEXT("TessellationMultiplier"), 1, 1, 0, 0, 0));
 	Outputs.Add(FExpressionOutput(TEXT("SubsurfaceColor"), 1, 1, 1, 1, 0));
 	Outputs.Add(FExpressionOutput(TEXT("ClearCoat"), 1, 1, 0, 0, 0));
 	Outputs.Add(FExpressionOutput(TEXT("ClearCoatRoughness"), 1, 1, 0, 0, 0));
@@ -5890,8 +5894,6 @@ void UMaterialExpressionBreakMaterialAttributes::Serialize(FStructuredArchive::F
 		Outputs[OutputIndex].SetMask(1, 1, 1, 1, 0); ++OutputIndex; // Normal
 		Outputs[OutputIndex].SetMask(1, 1, 1, 1, 0); ++OutputIndex; // Tangent
 		Outputs[OutputIndex].SetMask(1, 1, 1, 1, 0); ++OutputIndex; // WorldPositionOffset
-		Outputs[OutputIndex].SetMask(1, 1, 1, 1, 0); ++OutputIndex; // WorldDisplacement
-		Outputs[OutputIndex].SetMask(1, 1, 0, 0, 0); ++OutputIndex; // TessellationMultiplier
 		Outputs[OutputIndex].SetMask(1, 1, 1, 1, 0); ++OutputIndex; // SubsurfaceColor
 		Outputs[OutputIndex].SetMask(1, 1, 0, 0, 0); ++OutputIndex; // ClearCoat
 		Outputs[OutputIndex].SetMask(1, 1, 0, 0, 0); ++OutputIndex; // ClearCoatRoughness 
@@ -5927,24 +5929,22 @@ static void BuildPropertyToIOIndexMap()
 		PropertyToIOIndexMap.Add(MP_Normal,					8);
 		PropertyToIOIndexMap.Add(MP_Tangent,				9);
 		PropertyToIOIndexMap.Add(MP_WorldPositionOffset,	10);
-		PropertyToIOIndexMap.Add(MP_WorldDisplacement,		11);
-		PropertyToIOIndexMap.Add(MP_TessellationMultiplier, 12);
-		PropertyToIOIndexMap.Add(MP_SubsurfaceColor,		13);
-		PropertyToIOIndexMap.Add(MP_CustomData0,			14);
-		PropertyToIOIndexMap.Add(MP_CustomData1,			15);
-		PropertyToIOIndexMap.Add(MP_AmbientOcclusion,		16);
-		PropertyToIOIndexMap.Add(MP_Refraction,				17);
-		PropertyToIOIndexMap.Add(MP_CustomizedUVs0,			18);
-		PropertyToIOIndexMap.Add(MP_CustomizedUVs1,			19);
-		PropertyToIOIndexMap.Add(MP_CustomizedUVs2,			20);
-		PropertyToIOIndexMap.Add(MP_CustomizedUVs3,			21);
-		PropertyToIOIndexMap.Add(MP_CustomizedUVs4,			22);
-		PropertyToIOIndexMap.Add(MP_CustomizedUVs5,			23);
-		PropertyToIOIndexMap.Add(MP_CustomizedUVs6,			24);
-		PropertyToIOIndexMap.Add(MP_CustomizedUVs7,			25);
-		PropertyToIOIndexMap.Add(MP_PixelDepthOffset,		26);
-		PropertyToIOIndexMap.Add(MP_ShadingModel,			27);
-		PropertyToIOIndexMap.Add(MP_FrontMaterial,			28);
+		PropertyToIOIndexMap.Add(MP_SubsurfaceColor,		11);
+		PropertyToIOIndexMap.Add(MP_CustomData0,			12);
+		PropertyToIOIndexMap.Add(MP_CustomData1,			13);
+		PropertyToIOIndexMap.Add(MP_AmbientOcclusion,		14);
+		PropertyToIOIndexMap.Add(MP_Refraction,				15);
+		PropertyToIOIndexMap.Add(MP_CustomizedUVs0,			16);
+		PropertyToIOIndexMap.Add(MP_CustomizedUVs1,			17);
+		PropertyToIOIndexMap.Add(MP_CustomizedUVs2,			18);
+		PropertyToIOIndexMap.Add(MP_CustomizedUVs3,			19);
+		PropertyToIOIndexMap.Add(MP_CustomizedUVs4,			20);
+		PropertyToIOIndexMap.Add(MP_CustomizedUVs5,			21);
+		PropertyToIOIndexMap.Add(MP_CustomizedUVs6,			22);
+		PropertyToIOIndexMap.Add(MP_CustomizedUVs7,			23);
+		PropertyToIOIndexMap.Add(MP_PixelDepthOffset,		24);
+		PropertyToIOIndexMap.Add(MP_ShadingModel,			25);
+		PropertyToIOIndexMap.Add(MP_FrontMaterial,			26);
 	}
 }
 
diff --git a/Engine/Source/Runtime/Engine/Private/Materials/MaterialShared.cpp b/Engine/Source/Runtime/Engine/Private/Materials/MaterialShared.cpp
--- a/Engine/Source/Runtime/Engine/Private/Materials/MaterialShared.cpp
+++ b/Engine/Source/Runtime/Engine/Private/Materials/MaterialShared.cpp
@@ -8,6 +8,7 @@
 #include "Stats/StatsMisc.h"
 #include "UObject/CoreObjectVersion.h"
 #include "UObject/FrameworkObjectVersion.h"
+#include "UObject/UE5MainStreamObjectVersion.h"
 #include "Misc/App.h"
 #include "UObject/UObjectHash.h"
 #include "LocalVertexFactory.h"
@@ -4081,7 +4082,7 @@ bool UMaterialInterface::IsTextureReferencedByProperty(EMaterialProperty InPrope
 //Reorder the output index for any FExpressionInput connected to a UMaterialExpressionBreakMaterialAttributes.
 //If the order of pins in the material results or the make/break attributes nodes changes 
 //then the OutputIndex stored in any FExpressionInput coming from UMaterialExpressionBreakMaterialAttributes will be wrong and needs reordering.
-void DoMaterialAttributeReorder(FExpressionInput* Input, int32 UEVer, int32 RenderObjVer)
+void DoMaterialAttributeReorder(FExpressionInput* Input, int32 UEVer, int32 RenderObjVer, int32 UE5MainVer)
 {
 	if( Input && Input->Expression && Input->Expression->IsA(UMaterialExpressionBreakMaterialAttributes::StaticClass()) )
 	{
@@ -4125,6 +4126,20 @@ void DoMaterialAttributeReorder(FExpressionInput* Input, int32 UEVer, int32 Rend
 				++Input->OutputIndex;
 			}
 		}
+
+		if (UE5MainVer < FUE5MainStreamObjectVersion::RemovingTessellationParameters)
+		{
+			// Removing MP_WorldDisplacement (11) and MP_TessellationMultiplier (12)
+			if (Input->OutputIndex == 11 || Input->OutputIndex == 12)
+			{
+				Input->Expression = nullptr;
+				Input->ExpressionName = NAME_None;
+			}
+			else if (Input->OutputIndex >= 13)
+			{
+				Input->OutputIndex -= 2;
+			}
+		}
 	}
 }
 #endif // WITH_EDITORONLY_DATA
@@ -4308,8 +4323,6 @@ void FMaterialAttributeDefinitionMap::InitializeAttributeMap()
 
 	// Advanced attributes
 	Add(FGuid(0xF905F895, 0xD5814314, 0x916D2434, 0x8C40CE9E), TEXT("WorldPositionOffset"),		MP_WorldPositionOffset,		MCT_Float3,	FVector4(0,0,0,0),	SF_Vertex);
-	Add(FGuid(0x2091ECA2, 0xB59248EE, 0x8E2CD578, 0xD371926D), TEXT("WorldDisplacement"),		MP_WorldDisplacement,		MCT_Float3,	FVector4(0,0,0,0),	SF_Domain);
-	Add(FGuid(0xA0119D44, 0xC456450D, 0x9C39C933, 0x1F72D8D1), TEXT("TessellationMultiplier"),	MP_TessellationMultiplier,	MCT_Float,	FVector4(1,0,0,0),	SF_Hull);
 	Add(FGuid(0x5B8FC679, 0x51CE4082, 0x9D777BEE, 0xF4F72C44), TEXT("SubsurfaceColor"),			MP_SubsurfaceColor,			MCT_Float3,	FVector4(1,1,1,0),	SF_Pixel);
 	Add(FGuid(0x9E502E69, 0x3C8F48FA, 0x94645CFD, 0x28E5428D), TEXT("ClearCoat"),				MP_CustomData0,				MCT_Float,	FVector4(1,0,0,0),	SF_Pixel);
 	Add(FGuid(0xBE4F2FFD, 0x12FC4296, 0xB0124EEA, 0x12C28D92), TEXT("ClearCoatRoughness"),		MP_CustomData1,				MCT_Float,	FVector4(.1,0,0,0),	SF_Pixel);
@@ -4426,10 +4439,6 @@ FText FMaterialAttributeDefinitionMap::GetAttributeOverrideForMaterial(const FGu
 		return LOCTEXT("Tangent", "Tangent");
 	case MP_WorldPositionOffset:
 		return Material->IsUIMaterial() ? LOCTEXT("ScreenPosition", "Screen Position") : LOCTEXT("WorldPositionOffset", "World Position Offset");
-	case MP_WorldDisplacement:
-		return LOCTEXT("WorldDisplacement", "World Displacement");
-	case MP_TessellationMultiplier:
-		return LOCTEXT("TessellationMultiplier", "Tessellation Multiplier");
 	case MP_SubsurfaceColor:
 		if (Material->MaterialDomain == MD_Volume)
 		{
diff --git a/Engine/Source/Runtime/Engine/Public/MaterialShared.h b/Engine/Source/Runtime/Engine/Public/MaterialShared.h
--- a/Engine/Source/Runtime/Engine/Public/MaterialShared.h
+++ b/Engine/Source/Runtime/Engine/Public/MaterialShared.h
@@ -2716,7 +2716,7 @@ ENGINE_API bool DoesMaterialUseTexture(const UMaterialInterface* Material,const
 
 #if WITH_EDITORONLY_DATA
 /** TODO - This can be removed whenever VER_UE4_MATERIAL_ATTRIBUTES_REORDERING is no longer relevant. */
-ENGINE_API void DoMaterialAttributeReorder(FExpressionInput* Input, int32 UEVer, int32 RenderObjVer);
+ENGINE_API void DoMaterialAttributeReorder(FExpressionInput* Input, int32 UEVer, int32 RenderObjVer, int32 UE5MainVer);
 #endif // WITH_EDITORONLY_DATA
 
 /**
diff --git a/Engine/Source/Runtime/Engine/Public/SceneTypes.h b/Engine/Source/Runtime/Engine/Public/SceneTypes.h
--- a/Engine/Source/Runtime/Engine/Public/SceneTypes.h
+++ b/Engine/Source/Runtime/Engine/Public/SceneTypes.h
@@ -164,8 +164,6 @@ enum EMaterialProperty
 	MP_Normal UMETA(DisplayName = "Normal"),
 	MP_Tangent UMETA(DisplayName = "Tangent"),
 	MP_WorldPositionOffset UMETA(Hidden),
-	MP_WorldDisplacement UMETA(Hidden),
-	MP_TessellationMultiplier UMETA(Hidden),
 	MP_SubsurfaceColor UMETA(DisplayName = "Subsurface"),
 	MP_CustomData0 UMETA(Hidden),
 	MP_CustomData1 UMETA(Hidden),
