Fixed crash during USD MDL material creation where a UTexture was being garbaged collected when it was still needed. Don't duplicate a texture that's in the transient package during VT conversion since it's not engine content and it's propably not RF_Standalone.

Also made sure each material has it's own package as the current setup was to have package for the module with all the materials in it which lacked a proper root object.

diff --git a/Engine/Plugins/Enterprise/MDLImporter/Source/MDLImporter/Private/MDLMaterialFactory.cpp b/Engine/Plugins/Enterprise/MDLImporter/Source/MDLImporter/Private/MDLMaterialFactory.cpp
--- a/Engine/Plugins/Enterprise/MDLImporter/Source/MDLImporter/Private/MDLMaterialFactory.cpp
+++ b/Engine/Plugins/Enterprise/MDLImporter/Source/MDLImporter/Private/MDLMaterialFactory.cpp
@@ -14,6 +14,7 @@
 #include "Logging/LogMacros.h"
 #include "MaterialEditingLibrary.h"
 #include "MaterialEditorUtilities.h"
+#include "ObjectTools.h"
 #include "PackageTools.h"
 #include "UObject/Package.h"
 
@@ -180,9 +181,6 @@ bool FMDLMaterialFactory::CreateMaterials(const FString& Filename, UObject* Pare
 {
 	CleanUp();
 
-	FString  MaterialPackageName = UPackageTools::SanitizePackageName(*(ParentPackage->GetName() / Materials.Name));
-	UObject* MaterialPackage     = CreatePackage(*MaterialPackageName);
-
 	for (const Mdl::FMaterial& MdlMaterial : Materials)
 	{
 		if (MdlMaterial.IsDisabled())
@@ -190,7 +188,23 @@ bool FMDLMaterialFactory::CreateMaterials(const FString& Filename, UObject* Pare
 			continue;
 		}
 
-		UMaterial* NewMaterial = NewObject<UMaterial>(MaterialPackage, UMaterial::StaticClass(), FName(*MdlMaterial.Name), Flags);
+		FString  MaterialPackageName = UPackageTools::SanitizePackageName(*(ParentPackage->GetName() / MdlMaterial.Name));
+		UObject* MaterialPackage = ParentPackage == GetTransientPackage() ? ParentPackage : CreatePackage(*MaterialPackageName);
+
+		FString MaterialName;
+
+		if ( MaterialPackage == GetTransientPackage() )
+		{
+			// Avoid conflicts in the transient package by adding the module name to the material name since it's not part of the package path
+			MaterialName = Materials.Name + TEXT("_") + MdlMaterial.Name;
+		}
+		else
+		{
+			MaterialName = MdlMaterial.Name;
+		}
+		MaterialName = ObjectTools::SanitizeObjectName( MaterialName );
+
+		UMaterial* NewMaterial = NewObject<UMaterial>(MaterialPackage, UMaterial::StaticClass(), FName(*MaterialName), Flags);
 		check(NewMaterial != nullptr);
 		NewMaterial->AssetImportData = NewObject<UAssetImportData>(NewMaterial, TEXT("AssetImportData"));
 		NewMaterial->AssetImportData->Update(Filename);
diff --git a/Engine/Plugins/Enterprise/MDLImporter/Source/MDLImporter/Private/generator/FunctionGenerator.h b/Engine/Plugins/Enterprise/MDLImporter/Source/MDLImporter/Private/generator/FunctionGenerator.h
--- a/Engine/Plugins/Enterprise/MDLImporter/Source/MDLImporter/Private/generator/FunctionGenerator.h
+++ b/Engine/Plugins/Enterprise/MDLImporter/Source/MDLImporter/Private/generator/FunctionGenerator.h
@@ -2525,10 +2525,10 @@ namespace Generator
 				    ArraySize * 2.0f / PI);  // multiply with the number of curveValues, to get an index
 
 				TArray<FMaterialExpressionConnection> ColorConnections;
-				ColorConnections.SetNum(Colors.Num());
+				ColorConnections.Reserve(Colors.Num());
 				for (int32 i = 0; i < Colors.Num(); i++)
 				{
-					ColorConnections[i] = {Colors[i]};
+					ColorConnections.Emplace(Colors[i]);
 				}
 
 				Color = NewMaterialExpressionLinearInterpolate(Function, ColorConnections, Z);
diff --git a/Engine/Source/Developer/AssetTools/Private/AssetVtConversion.cpp b/Engine/Source/Developer/AssetTools/Private/AssetVtConversion.cpp
--- a/Engine/Source/Developer/AssetTools/Private/AssetVtConversion.cpp
+++ b/Engine/Source/Developer/AssetTools/Private/AssetVtConversion.cpp
@@ -552,7 +552,7 @@ void FVTConversionWorker::DoConvert()
 		for (UTexture2D *Tex : Textures)
 		{
 			UTexture2D* TextureToUpdate = Tex;
-			if (TextureToUpdate->GetPathName().StartsWith(TEXT("/Engine/")))
+			if (TextureToUpdate->GetPathName().StartsWith(TEXT("/Engine/")) && TextureToUpdate->GetPackage() != GetTransientPackage())
 			{
 				// rather than modify engine content, create a copy and update that
 				// any materials that we modify will be updated to point to the copy
