Fixed a crash when importing MDL materials with non power of 2 textures.
During async texture compilation, the surface size is always 32*32 (default platform data) so our power of two check was always true.

Preventively did the same change in the GLTFTextureFactory.

diff --git a/Engine/Plugins/Editor/GLTFImporter/Source/GLTFImporter/Private/GLTFTextureFactory.cpp b/Engine/Plugins/Editor/GLTFImporter/Source/GLTFImporter/Private/GLTFTextureFactory.cpp
--- a/Engine/Plugins/Editor/GLTFImporter/Source/GLTFImporter/Private/GLTFTextureFactory.cpp
+++ b/Engine/Plugins/Editor/GLTFImporter/Source/GLTFImporter/Private/GLTFTextureFactory.cpp
@@ -133,9 +133,7 @@ GLTF::ITextureElement* FGLTFTextureFactory::CreateTexture(const GLTF::FTexture&
 	{
 		TextureMipGenSettings MipGenSettings = TMGS_FromTextureGroup;
 
-		const int Width  = Texture->GetSurfaceWidth() > 0 ? (int)Texture->GetSurfaceWidth() : Texture->Source.GetSizeX();
-		const int Height = Texture->GetSurfaceHeight() > 0 ? (int)Texture->GetSurfaceHeight() : Texture->Source.GetSizeY();
-		if (!FMath::IsPowerOfTwo(Width) || !FMath::IsPowerOfTwo(Height))
+		if (!Texture->Source.IsPowerOfTwo())
 		{
 			MipGenSettings = TMGS_NoMipmaps;
 
diff --git a/Engine/Plugins/Enterprise/MDLImporter/Source/MDLImporter/Private/generator/MaterialTextureFactory.cpp b/Engine/Plugins/Enterprise/MDLImporter/Source/MDLImporter/Private/generator/MaterialTextureFactory.cpp
--- a/Engine/Plugins/Enterprise/MDLImporter/Source/MDLImporter/Private/generator/MaterialTextureFactory.cpp
+++ b/Engine/Plugins/Enterprise/MDLImporter/Source/MDLImporter/Private/generator/MaterialTextureFactory.cpp
@@ -93,9 +93,7 @@ namespace Generator
 	{
 		TextureMipGenSettings MipGenSettings = Property.MipGenSettings;
 
-		const int Width = Texture->GetSurfaceWidth() > 0 ? (int)Texture->GetSurfaceWidth() : Texture->Source.GetSizeX();
-		const int Height = Texture->GetSurfaceHeight() > 0 ? (int)Texture->GetSurfaceHeight() : Texture->Source.GetSizeY();
-		if (!FMath::IsPowerOfTwo(Width) || !FMath::IsPowerOfTwo(Height))
+		if (!Texture->Source.IsPowerOfTwo())
 		{
 			MipGenSettings = TMGS_NoMipmaps;
 
