Fixed usage of TCHAR_TO_ANSI in various plugins.

diff --git a/Engine/Plugins/Enterprise/DatasmithFBXImporter/Source/DatasmithFBXTranslator/Private/DatasmithFBXFileImporter.cpp b/Engine/Plugins/Enterprise/DatasmithFBXImporter/Source/DatasmithFBXTranslator/Private/DatasmithFBXFileImporter.cpp
--- a/Engine/Plugins/Enterprise/DatasmithFBXImporter/Source/DatasmithFBXTranslator/Private/DatasmithFBXFileImporter.cpp
+++ b/Engine/Plugins/Enterprise/DatasmithFBXImporter/Source/DatasmithFBXTranslator/Private/DatasmithFBXFileImporter.cpp
@@ -827,7 +827,7 @@ void FDatasmithFBXFileImporter::AddCurvesForProperty(FbxProperty InProperty, Fbx
 {
 	if (InProperty.IsValid())
 	{
-		auto PropertyName = ANSI_TO_TCHAR(InProperty.GetNameAsCStr());
+		auto PropertyName = StringCast<ANSICHAR>(InProperty.GetNameAsCStr());
 
 		FbxAnimCurveNode* CurveNode = InProperty.GetCurveNode(InLayer);
 
@@ -837,7 +837,7 @@ void FDatasmithFBXFileImporter::AddCurvesForProperty(FbxProperty InProperty, Fbx
 			for (uint8 Channel = 0; Channel < CurveNode->GetChannelsCount(); Channel++)
 			{
 				uint32 CurveCount = CurveNode->GetCurveCount(Channel);
-				UE_LOG(LogDatasmithFBXImport, Verbose, TEXT("\tFound %d curves for property %s, channel %d"), CurveCount, PropertyName, Channel);
+				UE_LOG(LogDatasmithFBXImport, Verbose, TEXT("\tFound %d curves for property %s, channel %d"), CurveCount, PropertyName.Get(), Channel);
 
 				for (uint32 CurveIndex = 0; CurveIndex < CurveCount; CurveIndex++)
 				{
diff --git a/Engine/Plugins/Enterprise/MDLImporter/Source/MDLImporter/Private/mdl/ApiContext.cpp b/Engine/Plugins/Enterprise/MDLImporter/Source/MDLImporter/Private/mdl/ApiContext.cpp
--- a/Engine/Plugins/Enterprise/MDLImporter/Source/MDLImporter/Private/mdl/ApiContext.cpp
+++ b/Engine/Plugins/Enterprise/MDLImporter/Source/MDLImporter/Private/mdl/ApiContext.cpp
@@ -208,8 +208,7 @@ namespace Mdl
 		}
 
 		const FString AbsolutePath = NormalizeDirectoryPathForMdl(*ModulesPath);
-		const char*   PathANSI           = TCHAR_TO_ANSI(*AbsolutePath);
-		MDL_CHECK_RESULT()               = ConfigHandle->add_mdl_path(PathANSI);
+		MDL_CHECK_RESULT() = ConfigHandle->add_mdl_path(TCHAR_TO_ANSI(*AbsolutePath));
 	}
 
 	void FApiContext::RemoveSearchPath(const FString& ModulesPath)
@@ -220,8 +219,7 @@ namespace Mdl
 		}
 
 		const FString AbsolutePath = NormalizeDirectoryPathForMdl(*ModulesPath);
-		const char*   PathANSI           = TCHAR_TO_ANSI(*AbsolutePath);
-		MDL_CHECK_RESULT()               = ConfigHandle->remove_mdl_path(PathANSI);
+		MDL_CHECK_RESULT() = ConfigHandle->remove_mdl_path(TCHAR_TO_ANSI(*AbsolutePath));
 	}
 
 	void FApiContext::AddResourceSearchPath(const FString& ResourcesPath)
@@ -232,8 +230,7 @@ namespace Mdl
 		}
 
 		const FString AbsolutePath = NormalizeDirectoryPathForMdl(*ResourcesPath);
-		const char*   PathANSI     = TCHAR_TO_ANSI(*AbsolutePath);
-		MDL_CHECK_RESULT()         = ConfigHandle->add_resource_path(PathANSI);
+		MDL_CHECK_RESULT() = ConfigHandle->add_resource_path(TCHAR_TO_ANSI(*AbsolutePath));
 	}
 
 	void FApiContext::RemoveResourceSearchPath(const FString& ResourcesPath)
@@ -244,8 +241,7 @@ namespace Mdl
 		}
 
 		const FString AbsolutePath = NormalizeDirectoryPathForMdl(*ResourcesPath);
-		const char*   PathANSI     = TCHAR_TO_ANSI(*AbsolutePath);
-		MDL_CHECK_RESULT()         = ConfigHandle->remove_resource_path(PathANSI);
+		MDL_CHECK_RESULT() = ConfigHandle->remove_resource_path(TCHAR_TO_ANSI(*AbsolutePath));
 	}
 
 	bool FApiContext::LoadModule(const FString& InModuleName, FMaterialCollection& OutMaterials)
diff --git a/Engine/Plugins/Media/AppleProResMedia/Source/AppleProResMedia/Private/AppleProResEncoderProtocol.cpp b/Engine/Plugins/Media/AppleProResMedia/Source/AppleProResMedia/Private/AppleProResEncoderProtocol.cpp
--- a/Engine/Plugins/Media/AppleProResMedia/Source/AppleProResMedia/Private/AppleProResEncoderProtocol.cpp
+++ b/Engine/Plugins/Media/AppleProResMedia/Source/AppleProResMedia/Private/AppleProResEncoderProtocol.cpp
@@ -97,12 +97,12 @@ bool UAppleProResEncoderProtocol::CreateFile(const FString& InFilename)
 		return false;
 	}
 
-	ANSICHAR *Filename = TCHAR_TO_ANSI(*FPaths::ConvertRelativePathToFull(InFilename));
+	const auto Filename = StringCast<ANSICHAR>(*FPaths::ConvertRelativePathToFull(InFilename));
 
 	Encoder = PROpenEncoder(NumberOfEncodingThreads, nullptr);
 
 	PRStatus status = 0;
-	status = ProResFileWriterCreate(Filename, &FileWriter);
+	status = ProResFileWriterCreate(Filename.Get(), &FileWriter);
 
 	if (status != 0)
 	{
diff --git a/Engine/Plugins/Runtime/nDisplay/Source/DisplayClusterProjection/Private/Policy/EasyBlend/Windows/DX11/DisplayClusterProjectionEasyBlendViewAdapterDX11.cpp b/Engine/Plugins/Runtime/nDisplay/Source/DisplayClusterProjection/Private/Policy/EasyBlend/Windows/DX11/DisplayClusterProjectionEasyBlendViewAdapterDX11.cpp
--- a/Engine/Plugins/Runtime/nDisplay/Source/DisplayClusterProjection/Private/Policy/EasyBlend/Windows/DX11/DisplayClusterProjectionEasyBlendViewAdapterDX11.cpp
+++ b/Engine/Plugins/Runtime/nDisplay/Source/DisplayClusterProjection/Private/Policy/EasyBlend/Windows/DX11/DisplayClusterProjectionEasyBlendViewAdapterDX11.cpp
@@ -96,7 +96,7 @@ bool FDisplayClusterProjectionEasyBlendViewAdapterDX11::Initialize(const FString
 		TRACE_CPUPROFILER_EVENT_SCOPE(nDisplay EasyBlend::Initialize);
 
 		// Initialize EasyBlend data for each view
-		const char* const FileName = TCHAR_TO_ANSI(*File);
+		const auto FileName = StringCast<ANSICHAR>(*File);
 		for (FViewData& It : Views)
 		{
 			// Initialize the mesh data
@@ -104,7 +104,7 @@ bool FDisplayClusterProjectionEasyBlendViewAdapterDX11::Initialize(const FString
 				FScopeLock lock(&DllAccessCS);
 
 				check(DisplayClusterProjectionEasyBlendLibraryDX11::EasyBlendInitializeFunc);
-				const EasyBlendSDKDXError Result = DisplayClusterProjectionEasyBlendLibraryDX11::EasyBlendInitializeFunc(FileName, It.EasyBlendMeshData.Get());
+				const EasyBlendSDKDXError Result = DisplayClusterProjectionEasyBlendLibraryDX11::EasyBlendInitializeFunc(FileName.Get(), It.EasyBlendMeshData.Get());
 				if (!EasyBlendSDKDX_SUCCEEDED(Result))
 				{
 					UE_LOG(LogDisplayClusterProjectionEasyBlend, Error, TEXT("Couldn't initialize EasyBlend internals"));
