nop fix comments

diff --git a/Engine/Plugins/Enterprise/MDLImporter/Source/MDLImporter/Private/common/TextureSource.h b/Engine/Plugins/Enterprise/MDLImporter/Source/MDLImporter/Private/common/TextureSource.h
--- a/Engine/Plugins/Enterprise/MDLImporter/Source/MDLImporter/Private/common/TextureSource.h
+++ b/Engine/Plugins/Enterprise/MDLImporter/Source/MDLImporter/Private/common/TextureSource.h
@@ -21,7 +21,10 @@ namespace Common
 	// Converts a float texture to a corresponding texture source.
 	inline FTextureSource* CreateTextureSource(const float* InData, int InWidth, int InHeight, int InChannels, bool bFlipY)
 	{
-		// @@!! TODO use FImage / CopyImage
+		// @todo Oodle : use FImageView / CopyImage
+		//  just make an FImageView on the float * and use TextureSource->Init() from ImageView
+		//  this whole function could be 2 lines
+		// 
 		//	also why are we converting to G8 or RGBA16 here ?
 		// just use F32 ETextureSourceFormat ?
 
@@ -85,7 +88,7 @@ namespace Common
 						int DstOffset = Y * InWidth * 4 + X * 4;
 						int SrcOffset = SrcY * InWidth * 4 + X * 4;
 
-						// @@!! use FColor::QuantizeUNormFloatTo16
+						// @todo Oodle: use FColor::QuantizeUNormFloatTo16
 						//	or just use FImage
 						check(IsValueInRange(InData[SrcOffset]));
 						check(IsValueInRange(InData[SrcOffset + 1]));
diff --git a/Engine/Plugins/Media/PixelStreaming/Source/PixelStreaming/Private/PixelStreamingTextureSource.cpp b/Engine/Plugins/Media/PixelStreaming/Source/PixelStreaming/Private/PixelStreamingTextureSource.cpp
--- a/Engine/Plugins/Media/PixelStreaming/Source/PixelStreaming/Private/PixelStreamingTextureSource.cpp
+++ b/Engine/Plugins/Media/PixelStreaming/Source/PixelStreaming/Private/PixelStreamingTextureSource.cpp
@@ -4,11 +4,6 @@
 #include "Stats.h"
 #include "Utils.h"
 
-/**
-* @@!!
-* this is PixelStreamingTextureSource.cpp
-* the cpp is named wrong
-*/
 
 /*
  * --------- IBackBufferTextureSource ------------------------------------------
diff --git a/Engine/Plugins/MovieScene/MovieRenderPipeline/Source/MovieRenderPipelineCore/Private/MoviePipelineImageQuantization.cpp b/Engine/Plugins/MovieScene/MovieRenderPipeline/Source/MovieRenderPipelineCore/Private/MoviePipelineImageQuantization.cpp
--- a/Engine/Plugins/MovieScene/MovieRenderPipeline/Source/MovieRenderPipelineCore/Private/MoviePipelineImageQuantization.cpp
+++ b/Engine/Plugins/MovieScene/MovieRenderPipeline/Source/MovieRenderPipelineCore/Private/MoviePipelineImageQuantization.cpp
@@ -14,7 +14,7 @@ namespace MoviePipeline
 DECLARE_CYCLE_STAT(TEXT("STAT_MoviePipeline_ImageQuantization"), STAT_ImageQuantization, STATGROUP_MoviePipeline);
 
 
-// @@!! TODO: get rid of this whole file
+// @todo Oodle : can get rid of this whole file
 //		use FImage, just call CopyImage
 
 
diff --git a/Engine/Source/Developer/TextureCompressor/Private/TextureCompressorModule.cpp b/Engine/Source/Developer/TextureCompressor/Private/TextureCompressorModule.cpp
--- a/Engine/Source/Developer/TextureCompressor/Private/TextureCompressorModule.cpp
+++ b/Engine/Source/Developer/TextureCompressor/Private/TextureCompressorModule.cpp
@@ -1470,7 +1470,7 @@ static FVector TransformWorldToSideSpace(uint32 CubemapFace, FVector InDirection
 	return Ret;
 }
 
-FVector ComputeSSCubeDirectionAtTexelCenter(uint32 x, uint32 y, float InvSideExtent)
+static FVector ComputeSSCubeDirectionAtTexelCenter(uint32 x, uint32 y, float InvSideExtent)
 {
 	// center of the texels
 	FVector DirectionSS((x + 0.5f) * InvSideExtent * 2 - 1, (y + 0.5f) * InvSideExtent * 2 - 1, 1);
diff --git a/Engine/Source/Runtime/ImageWriteQueue/Public/ImagePixelData.h b/Engine/Source/Runtime/ImageWriteQueue/Public/ImagePixelData.h
--- a/Engine/Source/Runtime/ImageWriteQueue/Public/ImagePixelData.h
+++ b/Engine/Source/Runtime/ImageWriteQueue/Public/ImagePixelData.h
@@ -28,7 +28,7 @@ struct IImagePixelDataPayload
 typedef TSharedPtr<IImagePixelDataPayload, ESPMode::ThreadSafe> FImagePixelPayloadPtr;
 
 
-// todo : use ImageCore FImage instead
+// @todo Oodle : use ImageCore FImage instead
 //  get rid of this whole class and TImagePixelData too
 //	just use an FImage instead
 struct FImagePixelData
