Fix MDLImporter when Edit & Continue is enabled

 - Cannot use __LINE__ as a compile-time constant because it cannot be used with Edit & Continue (/ZI)
 - see C2975 in MSDN

#rnx

diff --git a/Engine/Plugins/Enterprise/MDLImporter/Source/MDLImporter/MDLImporter.Build.cs b/Engine/Plugins/Enterprise/MDLImporter/Source/MDLImporter/MDLImporter.Build.cs
--- a/Engine/Plugins/Enterprise/MDLImporter/Source/MDLImporter/MDLImporter.Build.cs
+++ b/Engine/Plugins/Enterprise/MDLImporter/Source/MDLImporter/MDLImporter.Build.cs
@@ -4,6 +4,7 @@ using UnrealBuildTool;
 using System;
 using System.IO;
 using System.Collections.Generic;
+using EpicGames.Core;
 
 namespace UnrealBuildTool.Rules
 {
@@ -67,8 +68,14 @@ namespace UnrealBuildTool.Rules
                     PublicDelayLoadDLLs.Add(RuntimeModuleName);
                     RuntimeDependencies.Add(ModulePath);
                 }
-            }
-            else if (Target.Platform == UnrealTargetPlatform.Mac)
+
+				// MSVC does not allow __LINE__ in template parameters when edit & continue is enabled (see C2975 @ MSDN)
+				if (Target.bSupportEditAndContinue && Target.WindowsPlatform.Compiler.IsMSVC())
+				{
+					PrivateDefinitions.Add("MDL_MSVC_EDITCONTINUE_WORKAROUND=1");
+				}
+			}
+			else if (Target.Platform == UnrealTargetPlatform.Mac)
             {
                 RuntimeModuleNames.Add("libmdl_sdk.so");
                 RuntimeModuleNames.Add("nv_freeimage.so");
diff --git a/Engine/Plugins/Enterprise/MDLImporter/Source/MDLImporter/Private/mdl/Common.h b/Engine/Plugins/Enterprise/MDLImporter/Source/MDLImporter/Private/mdl/Common.h
--- a/Engine/Plugins/Enterprise/MDLImporter/Source/MDLImporter/Private/mdl/Common.h
+++ b/Engine/Plugins/Enterprise/MDLImporter/Source/MDLImporter/Private/mdl/Common.h
@@ -6,6 +6,10 @@
 
 #include "Misc/AssertionMacros.h"
 
+#if !defined(MDL_MSVC_EDITCONTINUE_WORKAROUND)
+	#define MDL_MSVC_EDITCONTINUE_WORKAROUND 0
+#endif
+
 namespace Mdl
 {
 	enum class EValueType
@@ -47,8 +51,14 @@ namespace Mdl
 
 #define MDL_TOKENPASTE(x, y) x##y
 #define MDL_TOKENPASTE2(x, y) MDL_TOKENPASTE(x, y)
-#define MDL_CHECK_RESULT() const Mdl::Detail::FCheckRes<__LINE__> MDL_TOKENPASTE2(res, __LINE__)
-#define MDL_CHECK_RESULT_MSG(msg) const Mdl::Detail::FCheckRes<__LINE__> MDL_TOKENPASTE2(res, __LINE__)
+#if MDL_MSVC_EDITCONTINUE_WORKAROUND
+	// MSVC does not allow __LINE__ in template parameters when edit & continue is enabled (see C2975 @ MSDN)
+	#define MDL_CHECK_RESULT() const Mdl::Detail::FCheckRes<0> MDL_TOKENPASTE2(res, __LINE__)
+	#define MDL_CHECK_RESULT_MSG(msg) const Mdl::Detail::FCheckRes<0> MDL_TOKENPASTE2(res, __LINE__)
+#else
+	#define MDL_CHECK_RESULT() const Mdl::Detail::FCheckRes<__LINE__> MDL_TOKENPASTE2(res, __LINE__)
+	#define MDL_CHECK_RESULT_MSG(msg) const Mdl::Detail::FCheckRes<__LINE__> MDL_TOKENPASTE2(res, __LINE__)
+#endif
 
 	namespace Detail
 	{
