[MDLImporter] Inclusive terminology fix

diff --git a/Engine/Plugins/Enterprise/MDLImporter/Source/MDLImporter/Private/MDLMaterialFactory.cpp b/Engine/Plugins/Enterprise/MDLImporter/Source/MDLImporter/Private/MDLMaterialFactory.cpp
--- a/Engine/Plugins/Enterprise/MDLImporter/Source/MDLImporter/Private/MDLMaterialFactory.cpp
+++ b/Engine/Plugins/Enterprise/MDLImporter/Source/MDLImporter/Private/MDLMaterialFactory.cpp
@@ -240,8 +240,8 @@ void FMDLMaterialFactory::PostImport(Mdl::FMaterialCollection& Materials)
 		CreatedMaterials.Add(Material);
 		FAssetRegistryModule::AssetCreated(Material);
 
-		const FString MasterMaterialName = FMDLMaterialSelector::ToString(MaterialSelector->GetMaterialType(MdlMaterial));
-		UE_LOG(LogMDLImporter, Log, TEXT("Created material %s based on %s"), *MdlMaterial.Name, *MasterMaterialName);
+		const FString MaterialTypeName = FMDLMaterialSelector::ToString(MaterialSelector->GetMaterialType(MdlMaterial));
+		UE_LOG(LogMDLImporter, Log, TEXT("Created material %s based on %s"), *MdlMaterial.Name, *MaterialTypeName);
 	}
 #endif
 }
@@ -253,8 +253,8 @@ void FMDLMaterialFactory::Reimport(const Mdl::FMaterial& MdlMaterial, UMaterial&
 
 	// Material.AssetImportData->Update(InFileName);
 
-	const FString MasterMaterialName = FMDLMaterialSelector::ToString(MaterialSelector->GetMaterialType(MdlMaterial));
-	UE_LOG(LogMDLImporter, Log, TEXT("Reimported material %s based on %s"), *MdlMaterial.Name, *MasterMaterialName);
+	const FString MaterialTypeName = FMDLMaterialSelector::ToString(MaterialSelector->GetMaterialType(MdlMaterial));
+	UE_LOG(LogMDLImporter, Log, TEXT("Reimported material %s based on %s"), *MdlMaterial.Name, *MaterialTypeName);
 #endif
 }
 
