MDL Importer - fixes for the new MDL SDK

- fixed material names (MDL SDK changed returned names to be full function signature with function parameters, set_materials_are_functions not affecting this)

- fixed base_file_texture, file_bump_texture, base_tangent_space_normal_texture

- fixed tex_lookup_float(3, 4), added tex_lookup_float2

diff --git a/Engine/Plugins/Enterprise/MDLImporter/Source/MDLImporter/MDLImporter.Build.cs b/Engine/Plugins/Enterprise/MDLImporter/Source/MDLImporter/MDLImporter.Build.cs
--- a/Engine/Plugins/Enterprise/MDLImporter/Source/MDLImporter/MDLImporter.Build.cs
+++ b/Engine/Plugins/Enterprise/MDLImporter/Source/MDLImporter/MDLImporter.Build.cs
@@ -17,6 +17,7 @@ namespace UnrealBuildTool.Rules
 
         public MDLImporter(ReadOnlyTargetRules Target) : base(Target)
         {
+	        OptimizeCode = CodeOptimization.Never;
 			bLegalToDistributeObjectCode = true;
 
             PrivateDependencyModuleNames.AddRange(
@@ -52,9 +53,11 @@ namespace UnrealBuildTool.Rules
 
             if (Target.Platform == UnrealTargetPlatform.Win64)
             {
+				RuntimeModuleNames.Add("dds.dll");
 				RuntimeModuleNames.Add("libmdl_sdk.dll");
+				RuntimeModuleNames.Add("mdl_distiller.dll");
 				RuntimeModuleNames.Add("nv_freeimage.dll");
-                    
+
                 foreach (string RuntimeModuleName in RuntimeModuleNames)
                 {
                     string ModulePath = Path.Combine(BinaryLibraryFolder, RuntimeModuleName);
@@ -77,9 +80,10 @@ namespace UnrealBuildTool.Rules
 			}
 			else if (Target.Platform == UnrealTargetPlatform.Mac)
             {
+                RuntimeModuleNames.Add("dds.so");
                 RuntimeModuleNames.Add("libmdl_sdk.so");
+                RuntimeModuleNames.Add("mdl_distiller.so");
                 RuntimeModuleNames.Add("nv_freeimage.so");
-                RuntimeModuleNames.Add("dds.so");
 
                 foreach (string RuntimeModuleName in RuntimeModuleNames)
                 {
@@ -99,7 +103,7 @@ namespace UnrealBuildTool.Rules
             if (Directory.Exists(ThirdPartyPath))
             {
                 //third party libraries
-                string[] Libs = { "mdl-sdk-334300.2228"};
+                string[] Libs = { "mdl-sdk-349500.8766a"};
                 foreach (string Lib in Libs)
                 {
                     string IncludePath = Path.Combine(ThirdPartyPath, Lib, "include");
diff --git a/Engine/Plugins/Enterprise/MDLImporter/Source/MDLImporter/Private/MDLImporter.cpp b/Engine/Plugins/Enterprise/MDLImporter/Source/MDLImporter/Private/MDLImporter.cpp
--- a/Engine/Plugins/Enterprise/MDLImporter/Source/MDLImporter/Private/MDLImporter.cpp
+++ b/Engine/Plugins/Enterprise/MDLImporter/Source/MDLImporter/Private/MDLImporter.cpp
@@ -394,7 +394,7 @@ bool FMDLImporter::Reimport(const FString& InFileName, const UMDLImporterOptions
 	UMaterial*      Material    = Cast<UMaterial>(OutMaterial);
 	MDLImporterImpl::ClearMaterial(Material);
 
-	const FString                   DbName       = Mdl::Util::GetMaterialDatabaseName(Materials.Name, MdlMaterial.Name, true);
+	const FString                   DbName       = Mdl::Util::GetMaterialDatabaseName(Materials.Name, MdlMaterial.BaseName, true);
 	const TMap<FString, UMaterial*> MaterialsMap = {{DbName, Material}};
 
 	MdlContext->GetDistiller()->SetBakingSettings(InImporterOptions.BakingResolution, InImporterOptions.BakingSamples);
diff --git a/Engine/Plugins/Enterprise/MDLImporter/Source/MDLImporter/Private/MDLMaterialFactory.cpp b/Engine/Plugins/Enterprise/MDLImporter/Source/MDLImporter/Private/MDLMaterialFactory.cpp
--- a/Engine/Plugins/Enterprise/MDLImporter/Source/MDLImporter/Private/MDLMaterialFactory.cpp
+++ b/Engine/Plugins/Enterprise/MDLImporter/Source/MDLImporter/Private/MDLMaterialFactory.cpp
@@ -211,7 +211,7 @@ bool FMDLMaterialFactory::CreateMaterials(const FString& Filename, UObject* Pare
 		NewMaterial->AssetImportData = NewObject<UAssetImportData>(NewMaterial, TEXT("AssetImportData"));
 		NewMaterial->AssetImportData->Update(Filename);
 
-		const FString DbName        = Mdl::Util::GetMaterialDatabaseName(Materials.Name, MdlMaterial.Name, true);
+		const FString DbName        = Mdl::Util::GetMaterialDatabaseName(Materials.Name, MdlMaterial.BaseName, true);
 		NameMaterialMap.Add(DbName) = NewMaterial;
 	}
 
@@ -228,7 +228,7 @@ void FMDLMaterialFactory::PostImport(Mdl::FMaterialCollection& Materials)
 			continue;
 		}
 
-		const FString DbName   = Mdl::Util::GetMaterialDatabaseName(Materials.Name, MdlMaterial.Name, true);
+		const FString DbName   = Mdl::Util::GetMaterialDatabaseName(Materials.Name, MdlMaterial.BaseName, true);
 		UMaterial*    Material = NameMaterialMap[DbName];
 		if (Material == nullptr)
 		{
diff --git a/Engine/Plugins/Enterprise/MDLImporter/Source/MDLImporter/Private/generator/ConstantExpressionFactory.cpp b/Engine/Plugins/Enterprise/MDLImporter/Source/MDLImporter/Private/generator/ConstantExpressionFactory.cpp
--- a/Engine/Plugins/Enterprise/MDLImporter/Source/MDLImporter/Private/generator/ConstantExpressionFactory.cpp
+++ b/Engine/Plugins/Enterprise/MDLImporter/Source/MDLImporter/Private/generator/ConstantExpressionFactory.cpp
@@ -20,6 +20,32 @@ MDLSDK_INCLUDES_END
 
 namespace Generator
 {
+	template<typename T>
+	FMaterialExpressionConnectionList FConstantExpressionFactory::AddNewMaterialExpressionConstantHelper(const mi::base::Handle<const mi::neuraylib::IValue_vector>& Value)
+	{
+		mi::base::Handle<T> Values[4];
+		for (mi::Size i = 0; i < Value->get_size(); ++i)
+		{
+			auto Handle = mi::base::make_handle(Value->get_value(i));
+			Values[i] = mi::base::make_handle(Handle->get_interface<T>());
+		}
+
+		switch (Value->get_size())
+		{
+			case 2:
+				return { AddExpression(NewMaterialExpressionConstant(CurrentMaterial, static_cast<float>(Values[0]->get_value()), static_cast<float>(Values[1]->get_value()))) };
+			case 3:
+				return { AddExpression(NewMaterialExpressionConstant(  //
+					CurrentMaterial, static_cast<float>(Values[0]->get_value()), static_cast<float>(Values[1]->get_value()), static_cast<float>(Values[2]->get_value()))) };
+			case 4:
+				return { AddExpression(NewMaterialExpressionConstant(  //
+					CurrentMaterial, static_cast<float>(Values[0]->get_value()), static_cast<float>(Values[1]->get_value()), static_cast<float>(Values[2]->get_value()), static_cast<float>(Values[3]->get_value()))) };
+		}
+		
+		return {};
+		
+	}
+
 	FConstantExpressionFactory::FConstantExpressionFactory()
 	    : TextureFactory(nullptr)
 	{
@@ -68,22 +94,29 @@ namespace Generator
 			{
 				auto Value = mi::base::make_handle(MDLConstant.get_interface<const mi::neuraylib::IValue_vector>());
 
-				mi::base::Handle<const mi::neuraylib::IValue_float> Values[4];
-				for (mi::Size i = 0; i < Value->get_size(); i++)
+				const mi::neuraylib::IType::Kind ElementKind = Value->get_type()->get_element_type()->get_kind();
+				FMaterialExpressionConnectionList Result;
+				switch (ElementKind)
 				{
-					auto Handle = mi::base::make_handle(Value->get_value(i));
-					Values[i]   = mi::base::make_handle(Handle->get_interface<const mi::neuraylib::IValue_float>());
+					case mi::neuraylib::IType::Kind::TK_BOOL:
+						Result = AddNewMaterialExpressionConstantHelper<const mi::neuraylib::IValue_bool>(Value);
+						break;
+					case mi::neuraylib::IType::Kind::TK_INT:
+						Result = AddNewMaterialExpressionConstantHelper<const mi::neuraylib::IValue_int>(Value);
+						break;
+					case mi::neuraylib::IType::Kind::TK_FLOAT:
+						Result = AddNewMaterialExpressionConstantHelper<const mi::neuraylib::IValue_float>(Value);
+						break;
+					case mi::neuraylib::IType::Kind::TK_DOUBLE:
+						Result = AddNewMaterialExpressionConstantHelper<const mi::neuraylib::IValue_double>(Value);
+						break;
+					default:
+						break;
 				}
-				switch (Value->get_size())
+
+				if (Result.Num() > 0)
 				{
-					case 2:
-						return {AddExpression(NewMaterialExpressionConstant(CurrentMaterial, Values[0]->get_value(), Values[1]->get_value()))};
-					case 3:
-						return {AddExpression(NewMaterialExpressionConstant(  //
-						    CurrentMaterial, Values[0]->get_value(), Values[1]->get_value(), Values[2]->get_value()))};
-					case 4:
-						return {AddExpression(NewMaterialExpressionConstant(  //
-						    CurrentMaterial, Values[0]->get_value(), Values[1]->get_value(), Values[2]->get_value(), Values[3]->get_value()))};
+					return Result;
 				}
 				break;
 			}
@@ -149,7 +182,7 @@ namespace Generator
 				{
 					Common::FTextureProperty Property;
 					Property.Path = Mdl::Util::GetTextureFileName(MDLTexture.get());
-					float Gamma = MDLTexture->get_effective_gamma();
+					float Gamma = MDLTexture->get_effective_gamma(0, 0);
 					Property.bIsSRGB = Gamma != 1.0f;
 					if (bProcessingNormapMap)
 					{
diff --git a/Engine/Plugins/Enterprise/MDLImporter/Source/MDLImporter/Private/generator/ConstantExpressionFactory.h b/Engine/Plugins/Enterprise/MDLImporter/Source/MDLImporter/Private/generator/ConstantExpressionFactory.h
--- a/Engine/Plugins/Enterprise/MDLImporter/Source/MDLImporter/Private/generator/ConstantExpressionFactory.h
+++ b/Engine/Plugins/Enterprise/MDLImporter/Source/MDLImporter/Private/generator/ConstantExpressionFactory.h
@@ -4,6 +4,12 @@
 
 #include "generator/BaseExpressionFactory.h"
 #include "generator/MaterialExpressionConnection.h"
+#include "mdl/MdlSdkDefines.h"
+
+MDLSDK_INCLUDES_START
+#include "mi/neuraylib/ivalue.h"
+#include "mi/neuraylib/ivector.h"
+MDLSDK_INCLUDES_END
 
 namespace mi
 {
@@ -37,6 +43,9 @@ namespace Generator
 			return Expressions.Add_GetRef(Expression);
 		}
 
+		template<typename T>
+		FMaterialExpressionConnectionList AddNewMaterialExpressionConstantHelper(const mi::base::Handle<const mi::neuraylib::IValue_vector>& Value);
+		
 		TArray<FMaterialExpressionHandle> Expressions;
 	};
 
diff --git a/Engine/Plugins/Enterprise/MDLImporter/Source/MDLImporter/Private/generator/FunctionGenerator.h b/Engine/Plugins/Enterprise/MDLImporter/Source/MDLImporter/Private/generator/FunctionGenerator.h
--- a/Engine/Plugins/Enterprise/MDLImporter/Source/MDLImporter/Private/generator/FunctionGenerator.h
+++ b/Engine/Plugins/Enterprise/MDLImporter/Source/MDLImporter/Private/generator/FunctionGenerator.h
@@ -692,6 +692,17 @@ namespace Generator
 			    Function, TEXT("normal"), EFunctionInputType::FunctionInput_Vector3, NewMaterialExpressionFunctionCall(Function, StateNormal, {}));
 			UMaterialExpressionFunctionInput* Clip = NewMaterialExpressionFunctionInput(
 			    Function, TEXT("clip"), EFunctionInputType::FunctionInput_StaticBool, NewMaterialExpressionStaticBool(Function, false));
+			UMaterialExpressionFunctionInput* BSpline = NewMaterialExpressionFunctionInput(
+			    Function, TEXT("b_spline"), EFunctionInputType::FunctionInput_StaticBool, NewMaterialExpressionStaticBool(Function, false));
+			UMaterialExpressionFunctionInput* AnimationStartTime = NewMaterialExpressionFunctionInput(
+			    Function, TEXT("animation_start_time"), EFunctionInputType::FunctionInput_Scalar, 0.0f);
+			UMaterialExpressionFunctionInput* AnimationCrop = NewMaterialExpressionFunctionInput(
+				Function, TEXT("animation_crop"), EFunctionInputType::FunctionInput_Vector2, {0.0f, 0.0f});
+			UMaterialExpressionFunctionInput* AnimationWrap = NewMaterialExpressionFunctionInput(
+			    Function, TEXT("animation_wrap"), EFunctionInputType::FunctionInput_Scalar, wrap_repeat); // wrap_mode
+			UMaterialExpressionFunctionInput* AnimationFps = NewMaterialExpressionFunctionInput(
+			    Function, TEXT("animation_fps"), EFunctionInputType::FunctionInput_Scalar, 30.0f);
+
 #if defined(USE_WORLD_ALIGNED_TEXTURES)
 #if defined(USE_WAT_AS_SCALAR)
 			UMaterialExpressionFunctionInput* UseWorldAlignedTexture =
@@ -951,6 +962,15 @@ namespace Generator
 			    NewMaterialExpressionFunctionInput(Function, TEXT("wrap_v"), EFunctionInputType::FunctionInput_Scalar, wrap_repeat);
 			UMaterialExpressionFunctionInput* Clip = NewMaterialExpressionFunctionInput(
 			    Function, TEXT("clip"), EFunctionInputType::FunctionInput_StaticBool, NewMaterialExpressionStaticBool(Function, false));
+			UMaterialExpressionFunctionInput* AnimationStartTime = NewMaterialExpressionFunctionInput(
+			    Function, TEXT("animation_start_time"), EFunctionInputType::FunctionInput_Scalar, 0.0f);
+			UMaterialExpressionFunctionInput* AnimationCrop = NewMaterialExpressionFunctionInput(
+				Function, TEXT("animation_crop"), EFunctionInputType::FunctionInput_Vector2, {0.0f, 0.0f});
+			UMaterialExpressionFunctionInput* AnimationWrap = NewMaterialExpressionFunctionInput(
+			    Function, TEXT("animation_wrap"), EFunctionInputType::FunctionInput_Scalar, wrap_repeat); // wrap_mode
+			UMaterialExpressionFunctionInput* AnimationFps = NewMaterialExpressionFunctionInput(
+			    Function, TEXT("animation_fps"), EFunctionInputType::FunctionInput_Scalar, 30.0f);
+
 #if defined(USE_WORLD_ALIGNED_TEXTURES)
 #if defined(USE_WAT_AS_SCALAR)
 			UMaterialExpressionFunctionInput* UseWorldAlignedTexture =
@@ -1672,6 +1692,15 @@ namespace Generator
 			    NewMaterialExpressionFunctionInput(Function, TEXT("scale"), EFunctionInputType::FunctionInput_Scalar, 1.0f);
 			UMaterialExpressionFunctionInput* Offset =
 			    NewMaterialExpressionFunctionInput(Function, TEXT("offset"), EFunctionInputType::FunctionInput_Scalar, 0.0f);
+			UMaterialExpressionFunctionInput* AnimationStartTime = NewMaterialExpressionFunctionInput(
+			    Function, TEXT("animation_start_time"), EFunctionInputType::FunctionInput_Scalar, 0.0f);
+			UMaterialExpressionFunctionInput* AnimationCrop = NewMaterialExpressionFunctionInput(
+				Function, TEXT("animation_crop"), EFunctionInputType::FunctionInput_Vector2, {0.0f, 0.0f});
+			UMaterialExpressionFunctionInput* AnimationWrap = NewMaterialExpressionFunctionInput(
+			    Function, TEXT("animation_wrap"), EFunctionInputType::FunctionInput_Scalar, wrap_repeat); // wrap_mode
+			UMaterialExpressionFunctionInput* AnimationFps = NewMaterialExpressionFunctionInput(
+			    Function, TEXT("animation_fps"), EFunctionInputType::FunctionInput_Scalar, 30.0f);
+
 #if defined(USE_WORLD_ALIGNED_TEXTURES)
 #if defined(USE_WAT_AS_SCALAR)
 			UMaterialExpressionFunctionInput* UseWorldAlignedTexture =
@@ -3507,11 +3536,46 @@ namespace Generator
 			    NewMaterialExpressionFunctionInput(Function, TEXT("crop_u"), EFunctionInputType::FunctionInput_Vector2, {0.0f, 1.0f});
 			UMaterialExpressionFunctionInput* CropV =
 			    NewMaterialExpressionFunctionInput(Function, TEXT("crop_v"), EFunctionInputType::FunctionInput_Vector2, {0.0f, 1.0f});
+			UMaterialExpressionFunctionInput* Frame =
+			    NewMaterialExpressionFunctionInput(Function, TEXT("frame"), EFunctionInputType::FunctionInput_Scalar, 0);
+
+			UMaterialExpressionMaterialFunctionCall* Sample =
+			    NewMaterialExpressionFunctionCall(Function, ImporterTextureSample, {Tex, Coord, WrapU, WrapV, CropU, CropV});
+
+			NewMaterialExpressionFunctionOutput(Function, TEXT("result"), {Sample, 1});// r in output index 1
+		}
+
+		void TexLookupFloat2(UMaterialFunction* Function, int32 ArraySize)
+		{
+			check(0 == ArraySize);
+			UMaterialFunction* ImporterTextureSample = LoadFunction(TEXT("mdlimporter_texture_sample"));
+
+			Function->Description = TEXT(
+			    "Returns the sampled texture value for the twodimensional coordinates coord given in normalized texture space in the range [0, 1)², where the wrap"
+			    "modes define the behavior for coordinates outside of that range. The crop parameters further define a sub - range on the texture that is actually"
+			    "used and that defines the normalized texture space in the range [0, 1)². The crop parameter defaults float2(0.0, 1.0) corresponds to the whole texture"
+			    "in the corresponding axis. A lookup on an invalid texture reference returns zero.");
+
+			UMaterialExpressionFunctionInput* Tex = NewMaterialExpressionFunctionInput(
+			    Function, TEXT("tex"), EFunctionInputType::FunctionInput_Texture2D, NewMaterialExpressionTextureObject(Function, nullptr));
+			UMaterialExpressionFunctionInput* Coord = NewMaterialExpressionFunctionInput(
+			    Function, TEXT("coord"), EFunctionInputType::FunctionInput_Vector2, NewMaterialExpressionTextureCoordinate(Function, 0));
+			UMaterialExpressionFunctionInput* WrapU =
+			    NewMaterialExpressionFunctionInput(Function, TEXT("wrap_u"), EFunctionInputType::FunctionInput_Scalar, wrap_repeat);
+			UMaterialExpressionFunctionInput* WrapV =
+			    NewMaterialExpressionFunctionInput(Function, TEXT("wrap_v"), EFunctionInputType::FunctionInput_Scalar, wrap_repeat);
+			UMaterialExpressionFunctionInput* CropU =
+			    NewMaterialExpressionFunctionInput(Function, TEXT("crop_u"), EFunctionInputType::FunctionInput_Vector2, {0.0f, 1.0f});
+			UMaterialExpressionFunctionInput* CropV =
+			    NewMaterialExpressionFunctionInput(Function, TEXT("crop_v"), EFunctionInputType::FunctionInput_Vector2, {0.0f, 1.0f});
+			UMaterialExpressionFunctionInput* Frame =
+			    NewMaterialExpressionFunctionInput(Function, TEXT("frame"), EFunctionInputType::FunctionInput_Scalar, 0);
 
 			UMaterialExpressionMaterialFunctionCall* Sample =
 			    NewMaterialExpressionFunctionCall(Function, ImporterTextureSample, {Tex, Coord, WrapU, WrapV, CropU, CropV});
 
-			NewMaterialExpressionFunctionOutput(Function, TEXT("result"), {Sample, 1});
+			UMaterialFunction* MakeFloat2 = LoadFunction(TEXT("/Engine/Functions/Engine_MaterialFunctions02/Utility"), TEXT("MakeFloat2.MakeFloat2"));
+			NewMaterialExpressionFunctionOutput(Function, TEXT("result"), NewMaterialExpressionFunctionCall(Function, MakeFloat2, {{Sample, 1}, {Sample, 2}})); // (r, b)
 		}
 
 		void TexLookupFloat3(UMaterialFunction* Function, int32 ArraySize)
@@ -3537,11 +3601,13 @@ namespace Generator
 			    NewMaterialExpressionFunctionInput(Function, TEXT("crop_u"), EFunctionInputType::FunctionInput_Vector2, {0.0f, 1.0f});
 			UMaterialExpressionFunctionInput* CropV =
 			    NewMaterialExpressionFunctionInput(Function, TEXT("crop_v"), EFunctionInputType::FunctionInput_Vector2, {0.0f, 1.0f});
+			UMaterialExpressionFunctionInput* Frame =
+			    NewMaterialExpressionFunctionInput(Function, TEXT("frame"), EFunctionInputType::FunctionInput_Scalar, 0);
 
 			UMaterialExpressionMaterialFunctionCall* Sample =
 			    NewMaterialExpressionFunctionCall(Function, ImporterTextureSample, {Tex, Coord, WrapU, WrapV, CropU, CropV});
 
-			NewMaterialExpressionFunctionOutput(Function, TEXT("result"), Sample);
+			NewMaterialExpressionFunctionOutput(Function, TEXT("result"), Sample); // rgb in output index 0
 		}
 
 		void TexLookupFloat4(UMaterialFunction* Function, int32 ArraySize)
@@ -3567,11 +3633,13 @@ namespace Generator
 			    NewMaterialExpressionFunctionInput(Function, TEXT("crop_u"), EFunctionInputType::FunctionInput_Vector2, {0.0f, 1.0f});
 			UMaterialExpressionFunctionInput* CropV =
 			    NewMaterialExpressionFunctionInput(Function, TEXT("crop_v"), EFunctionInputType::FunctionInput_Vector2, {0.0f, 1.0f});
+			UMaterialExpressionFunctionInput* Frame =
+			    NewMaterialExpressionFunctionInput(Function, TEXT("frame"), EFunctionInputType::FunctionInput_Scalar, 0);
 
 			UMaterialExpressionMaterialFunctionCall* Sample =
 			    NewMaterialExpressionFunctionCall(Function, ImporterTextureSample, {Tex, Coord, WrapU, WrapV, CropU, CropV});
 
-			NewMaterialExpressionFunctionOutput(Function, TEXT("result"), NewMaterialExpressionAppendVector(Function, {Sample, 0}, {Sample, 4}));
+			NewMaterialExpressionFunctionOutput(Function, TEXT("result"), NewMaterialExpressionAppendVector(Function, {Sample, 0}, {Sample, 4}));// Append(rba, a)
 		}
 
 		/*static*/ void ImporterAddDetailNormal(UMaterialFunction* Function, int32 ArraySize)
diff --git a/Engine/Plugins/Enterprise/MDLImporter/Source/MDLImporter/Private/generator/FunctionLoader.cpp b/Engine/Plugins/Enterprise/MDLImporter/Source/MDLImporter/Private/generator/FunctionLoader.cpp
--- a/Engine/Plugins/Enterprise/MDLImporter/Source/MDLImporter/Private/generator/FunctionLoader.cpp
+++ b/Engine/Plugins/Enterprise/MDLImporter/Source/MDLImporter/Private/generator/FunctionLoader.cpp
@@ -155,6 +155,7 @@ namespace Generator
 		FunctionGenerateMap.Add(TEXT("mdl_state_transform_vector"), {&FFunctionGenerator::StateTransformVector, 1});
 		FunctionGenerateMap.Add(TEXT("mdl_tex_lookup_color"), {&FFunctionGenerator::TexLookupFloat3, 7});
 		FunctionGenerateMap.Add(TEXT("mdl_tex_lookup_float"), {&FFunctionGenerator::TexLookupFloat, 2});
+		FunctionGenerateMap.Add(TEXT("mdl_tex_lookup_float2"), {&FFunctionGenerator::TexLookupFloat2, 2});
 		FunctionGenerateMap.Add(TEXT("mdl_tex_lookup_float3"), {&FFunctionGenerator::TexLookupFloat3, 7});
 		FunctionGenerateMap.Add(TEXT("mdl_tex_lookup_float4"), {&FFunctionGenerator::TexLookupFloat4, 7});
 		FunctionGenerateMap.Add(TEXT("mdlimporter_add_detail_normal"), {&FFunctionGenerator::ImporterAddDetailNormal, 1});
diff --git a/Engine/Plugins/Enterprise/MDLImporter/Source/MDLImporter/Private/generator/MaterialExpressionFactory.cpp b/Engine/Plugins/Enterprise/MDLImporter/Source/MDLImporter/Private/generator/MaterialExpressionFactory.cpp
--- a/Engine/Plugins/Enterprise/MDLImporter/Source/MDLImporter/Private/generator/MaterialExpressionFactory.cpp
+++ b/Engine/Plugins/Enterprise/MDLImporter/Source/MDLImporter/Private/generator/MaterialExpressionFactory.cpp
@@ -763,7 +763,7 @@ namespace Generator
 				check(Inputs[0].GetConnectionType() == EConnectionType::Expression);
 				check(Inputs[1].GetConnectionType() == EConnectionType::Expression);
 				check(Inputs[1].IsExpressionA<UMaterialExpressionConstant>());
-				check(FunctionName == TEXT("operator[](<0>[],int)"));
+				check(FunctionName == TEXT("operator[](<0>[],int)") || FunctionName == TEXT("operator[](%3C0%3E[],int)"));
 
 				const int32 Index = (int32)Cast<UMaterialExpressionConstant>(Inputs[1].GetExpressionUnused())->R;
 				CurrentMaterial->GetExpressionCollection().RemoveExpression(Inputs[1].GetExpressionUnused());
@@ -1069,6 +1069,7 @@ namespace Generator
 			case mi::neuraylib::IFunction_definition::DS_INTRINSIC_STATE_TRANSFORM_POINT:
 			case mi::neuraylib::IFunction_definition::DS_INTRINSIC_STATE_TRANSFORM_VECTOR:
 			case mi::neuraylib::IFunction_definition::DS_INTRINSIC_TEX_LOOKUP_COLOR:
+			case mi::neuraylib::IFunction_definition::DS_INTRINSIC_TEX_LOOKUP_FLOAT2:
 			case mi::neuraylib::IFunction_definition::DS_INTRINSIC_TEX_LOOKUP_FLOAT3:
 			case mi::neuraylib::IFunction_definition::DS_INTRINSIC_TEX_LOOKUP_FLOAT4:
 			{
diff --git a/Engine/Plugins/Enterprise/MDLImporter/Source/MDLImporter/Private/generator/ParameterExpressionFactory.cpp b/Engine/Plugins/Enterprise/MDLImporter/Source/MDLImporter/Private/generator/ParameterExpressionFactory.cpp
--- a/Engine/Plugins/Enterprise/MDLImporter/Source/MDLImporter/Private/generator/ParameterExpressionFactory.cpp
+++ b/Engine/Plugins/Enterprise/MDLImporter/Source/MDLImporter/Private/generator/ParameterExpressionFactory.cpp
@@ -392,7 +392,7 @@ namespace Generator
 				{
 					Common::FTextureProperty Property;
 					Property.Path    = Mdl::Util::GetTextureFileName(MDLTexture.get());
-					float Gamma      = MDLTexture->get_effective_gamma();
+					float Gamma      = MDLTexture->get_effective_gamma(0, 0);
 					Property.bIsSRGB = Gamma != 1.0;
 					if (!FPaths::GetExtension(Property.Path).IsEmpty())
 					{
diff --git a/Engine/Plugins/Enterprise/MDLImporter/Source/MDLImporter/Private/mdl/ApiContext.cpp b/Engine/Plugins/Enterprise/MDLImporter/Source/MDLImporter/Private/mdl/ApiContext.cpp
--- a/Engine/Plugins/Enterprise/MDLImporter/Source/MDLImporter/Private/mdl/ApiContext.cpp
+++ b/Engine/Plugins/Enterprise/MDLImporter/Source/MDLImporter/Private/mdl/ApiContext.cpp
@@ -146,7 +146,16 @@ namespace Mdl
 		mi::Sint32 Result = PluginConfig->load_plugin_library(TCHAR_TO_ANSI(*FreeimageModulePath));
 		if ( Result != 0)
 		{
-			UE_LOG(LogMDLImporter, Error, TEXT("mi::neuraylib::IMdl_compiler::load_plugin_library() failed with return code %d."), Result);
+			UE_LOG(LogMDLImporter, Error, TEXT("mi::neuraylib::IPlugin_configuration::load_plugin_library() failed with return code %d."), Result);
+			return false;
+		}
+
+		// Load the Distiller plugin.
+		FString DistillerModulePath = FPaths::Combine(LibrariesPath, TEXT("mdl_distiller" MI_BASE_DLL_FILE_EXT));
+		Result = PluginConfig->load_plugin_library(TCHAR_TO_ANSI(*DistillerModulePath));
+		if ( Result != 0)
+		{
+			UE_LOG(LogMDLImporter, Error, TEXT("mi::neuraylib::IPlugin_configuration::load_plugin_library() failed with return code %d."), Result);
 			return false;
 		}
 
@@ -278,17 +287,37 @@ namespace Mdl
 			{
 				const char* Name = Module->get_material(Index);
 				if (MaterialIsHidden(Name, *Transaction))
+				{
 					continue;
+				}
+
+				mi::base::Handle<const mi::neuraylib::IMaterial_definition> MaterialDefinition =
+				    mi::base::make_handle((*Transaction).access<mi::neuraylib::IMaterial_definition>(Name));
 
+				FString MdlName = ANSI_TO_TCHAR(MaterialDefinition->get_mdl_name());
+				FString MdlModuleName = ANSI_TO_TCHAR(MaterialDefinition->get_mdl_module_name());
+				FString MdlSimpleName = ANSI_TO_TCHAR(MaterialDefinition->get_mdl_simple_name());
+				
 				FMaterial& Material = OutMaterials.Create();
-				Material.Name       = ANSI_TO_TCHAR(Name);
-				Material.Id         = Index;
-				// strip Module name, format is: mdl::<module_name>::<material_name>
-				const int32 Found = Material.Name.Find(TEXT("::"), ESearchCase::IgnoreCase, ESearchDir::FromEnd);
+				Material.Id = Index;
+				
+				// Strip Module name, format is: mdl::<module_name>::<material_name>(<material-as-function parameter list>)
+				FString MaterialMdlDbName = ANSI_TO_TCHAR(Name);
+
+				int32 OpeningBracketIndex = INDEX_NONE;
+				MaterialMdlDbName.FindChar(TEXT('('), OpeningBracketIndex);
+
+				// Search module separator skipping function arguments - parameter types can also contain "::"
+				const int32 Found = MaterialMdlDbName.Find(TEXT("::"), ESearchCase::IgnoreCase, ESearchDir::FromEnd, OpeningBracketIndex);
 				if ((Found - 2) > 0)
 				{
-					Material.Name = Material.Name.Right(Material.Name.Len() - Found - 2);
+					MaterialMdlDbName = MaterialMdlDbName.Right(MaterialMdlDbName.Len() - Found - 2);
 				}
+
+				Material.BaseName = MaterialMdlDbName;
+
+				// Use material name without full signature (previously material db name returned from MDL SDK didn't contain parameter list)
+				Material.Name = MdlSimpleName;
 			}
 		}
 		MDL_CHECK_RESULT() = Transaction->commit();
diff --git a/Engine/Plugins/Enterprise/MDLImporter/Source/MDLImporter/Private/mdl/BakeParam.cpp b/Engine/Plugins/Enterprise/MDLImporter/Source/MDLImporter/Private/mdl/BakeParam.cpp
--- a/Engine/Plugins/Enterprise/MDLImporter/Source/MDLImporter/Private/mdl/BakeParam.cpp
+++ b/Engine/Plugins/Enterprise/MDLImporter/Source/MDLImporter/Private/mdl/BakeParam.cpp
@@ -97,7 +97,7 @@ namespace Mdl
 			RemapFunc(CanvasData);
 		}
 
-		mi::base::Handle<const mi::neuraylib::ITile> Tile(CanvasData.get_tile(0, 0));
+		mi::base::Handle<const mi::neuraylib::ITile> Tile(CanvasData.get_tile());
 		const mi::Float32*                           Buffer = static_cast<const mi::Float32*>(Tile->get_data());
 
 		BakedTextureData->SetName(Name);
diff --git a/Engine/Plugins/Enterprise/MDLImporter/Source/MDLImporter/Private/mdl/Material.h b/Engine/Plugins/Enterprise/MDLImporter/Source/MDLImporter/Private/mdl/Material.h
--- a/Engine/Plugins/Enterprise/MDLImporter/Source/MDLImporter/Private/mdl/Material.h
+++ b/Engine/Plugins/Enterprise/MDLImporter/Source/MDLImporter/Private/mdl/Material.h
@@ -146,7 +146,16 @@ namespace Mdl
 		};
 
 		uint32  Id;
-		FString Name;
+		FString Name; // Name used to identify Material
+
+		// Base name in the Mdl db(full name is something like ::<ModuleName>::<BaseName>)
+		// note: separated from Name, because with resent Mdl SDK changes db name includes full signature(i.e. material-as-function argument list)
+		FString BaseName; 
+
+		FString GetBaseName()
+		{
+			return BaseName;
+		}
 		// The preferred baking texture size.
 		int PreferredWidth;
 		int PreferredHeight;
diff --git a/Engine/Plugins/Enterprise/MDLImporter/Source/MDLImporter/Private/mdl/MaterialDistiller.cpp b/Engine/Plugins/Enterprise/MDLImporter/Source/MDLImporter/Private/mdl/MaterialDistiller.cpp
--- a/Engine/Plugins/Enterprise/MDLImporter/Source/MDLImporter/Private/mdl/MaterialDistiller.cpp
+++ b/Engine/Plugins/Enterprise/MDLImporter/Source/MDLImporter/Private/mdl/MaterialDistiller.cpp
@@ -4,7 +4,7 @@
 
 #include "MaterialDistiller.h"
 
-#define MDL_DEBUG_PRINT_MATERIAL 0
+#define MDL_DEBUG_PRINT_MATERIAL 3
 
 #include "ArgumentGetter.h"
 #include "BakeParam.h"
@@ -85,7 +85,7 @@ namespace Mdl
 				return;
 			}
 
-			mi::base::Handle<mi::neuraylib::ITile> Tile(Canvas->get_tile(0, 0));
+			mi::base::Handle<mi::neuraylib::ITile> Tile(Canvas->get_tile());
 			mi::Float32*                           Buffer = static_cast<mi::Float32*>(Tile->get_data());
 
 			const mi::Uint32 ChannelCount = 3;
@@ -109,7 +109,7 @@ namespace Mdl
 				return;
 			}
 
-			mi::base::Handle<mi::neuraylib::ITile> Tile(Canvas->get_tile(0, 0));
+			mi::base::Handle<mi::neuraylib::ITile> Tile(Canvas->get_tile());
 			mi::Float32*                           Buffer = static_cast<mi::Float32*>(Tile->get_data());
 
 			const mi::Uint32 ChannelCount = 1;
@@ -261,7 +261,7 @@ namespace Mdl
 			                   [&Material, MetersPerSceneUnit](mi::base::IInterface& ColorInterface)  //
 			                   {
 				                   mi::base::Handle<mi::neuraylib::ICanvas> Canvas(ColorInterface.get_interface<mi::neuraylib::ICanvas>());
-				                   mi::base::Handle<mi::neuraylib::ITile>   Tile(Canvas->get_tile(0, 0));
+				                   mi::base::Handle<mi::neuraylib::ITile>   Tile(Canvas->get_tile());
 								   const char* PixelType = Tile->get_type();
 				                   mi::Float32* Data = static_cast<mi::Float32*>(Tile->get_data());
 
@@ -381,7 +381,7 @@ namespace Mdl
 		int MaterialIndex = 0;
 		for (FMaterial& Material : Materials)
 		{
-			if (Material.Name.IsEmpty())
+			if (Material.IsDisabled())
 			{
 				continue;
 			}
@@ -390,7 +390,7 @@ namespace Mdl
 			mi::base::Handle<mi::neuraylib::ITransaction> Transaction(Scope->create_transaction());
 
 			// Compile the material before distillation
-			FString       MaterialDbName         = Mdl::Util::GetMaterialDatabaseName(Materials.Name, Material.Name);
+			FString       MaterialDbName         = Mdl::Util::GetMaterialDatabaseName(Materials.Name, Material.BaseName);
 			const FString MaterialInstanceDbName = Mdl::Util::GetMaterialInstanceName(MaterialDbName);
 			const FString MaterialCompiledDbName = MaterialDbName + TEXT("_compiled");
 			if (!Material.InstantiateFunction)
@@ -569,7 +569,7 @@ namespace Mdl
 		}
 		MDL_CHECK_RESULT() = Baker->bake_texture(Canvas.get(), 1);
 
-		mi::base::Handle<const mi::neuraylib::ITile> Tile(Canvas->get_tile(0, 0));
+		mi::base::Handle<const mi::neuraylib::ITile> Tile(Canvas->get_tile());
 		const mi::Float32*                           Buffer = static_cast<const mi::Float32*>(Tile->get_data());
 
 		check(4 >= ChannelCount);
diff --git a/Engine/Plugins/Enterprise/MDLImporter/Source/MDLImporter/Private/mdl/Utility.cpp b/Engine/Plugins/Enterprise/MDLImporter/Source/MDLImporter/Private/mdl/Utility.cpp
--- a/Engine/Plugins/Enterprise/MDLImporter/Source/MDLImporter/Private/mdl/Utility.cpp
+++ b/Engine/Plugins/Enterprise/MDLImporter/Source/MDLImporter/Private/mdl/Utility.cpp
@@ -63,6 +63,7 @@ namespace Mdl
 			check(Texture);
 			FString FileName = ANSI_TO_TCHAR(Texture->get_image());
 			FileName.RemoveFromStart("MI_default_image_");
+			FileName.RemoveFromEnd("_");
 			FPaths::NormalizeDirectoryName(FileName);
 			return FileName;
 		}
@@ -147,7 +148,7 @@ namespace Mdl
 		{
 			mi::neuraylib::IImage_api* ImageApi  = Neuray->get_api_component<mi::neuraylib::IImage_api>();
 			const char*                PixelType = GetPixelType(Channels);
-			mi::neuraylib::ICanvas*    Canvas    = ImageApi->create_canvas(PixelType, Width, Height, 0, 0, 1, false, Gamma);
+			mi::neuraylib::ICanvas*    Canvas    = ImageApi->create_canvas(PixelType, Width, Height, 1, false, Gamma);
 
 			return Canvas;
 		}
@@ -157,7 +158,7 @@ namespace Mdl
 		{
 			mi::neuraylib::IImage_api* ImageApi  = Neuray->get_api_component<mi::neuraylib::IImage_api>();
 			const char*                PixelType = GetPixelType(Channels);
-			mi::neuraylib::ICanvas*    Canvas    = ImageApi->create_canvas(PixelType, Width, Height, 0, 0, 1, false, Gamma);
+			mi::neuraylib::ICanvas*    Canvas    = ImageApi->create_canvas(PixelType, Width, Height, 1, false, Gamma);
 			MDL_CHECK_RESULT()                   = ImageApi->write_raw_pixels(Width, Height, Canvas, 0, 0, 0, Src, bFlipY, PixelType);
 
 			return Canvas;
@@ -185,7 +186,7 @@ namespace Mdl
 
 		int GetChannelCount(const mi::neuraylib::ICanvas& Canvas)
 		{
-			mi::base::Handle<const mi::neuraylib::ITile> Tile(Canvas.get_tile(0, 0));
+			mi::base::Handle<const mi::neuraylib::ITile> Tile(Canvas.get_tile());
 
 			const FString Type = ANSI_TO_TCHAR(Tile->get_type());
 			if (Type == TEXT("Float32"))
