MDL Importer - re-enabled code optimizations

diff --git a/Engine/Plugins/Enterprise/MDLImporter/Source/MDLImporter/MDLImporter.Build.cs b/Engine/Plugins/Enterprise/MDLImporter/Source/MDLImporter/MDLImporter.Build.cs
--- a/Engine/Plugins/Enterprise/MDLImporter/Source/MDLImporter/MDLImporter.Build.cs
+++ b/Engine/Plugins/Enterprise/MDLImporter/Source/MDLImporter/MDLImporter.Build.cs
@@ -17,7 +17,6 @@ namespace UnrealBuildTool.Rules
 
         public MDLImporter(ReadOnlyTargetRules Target) : base(Target)
         {
-	        OptimizeCode = CodeOptimization.Never;
 			bLegalToDistributeObjectCode = true;
 
             PrivateDependencyModuleNames.AddRange(
