Added dependency to MDL Linux binaries

diff --git a/Engine/Plugins/Enterprise/MDLImporter/Source/MDLImporter/MDLImporter.Build.cs b/Engine/Plugins/Enterprise/MDLImporter/Source/MDLImporter/MDLImporter.Build.cs
--- a/Engine/Plugins/Enterprise/MDLImporter/Source/MDLImporter/MDLImporter.Build.cs
+++ b/Engine/Plugins/Enterprise/MDLImporter/Source/MDLImporter/MDLImporter.Build.cs
@@ -77,7 +77,7 @@ namespace UnrealBuildTool.Rules
 					PrivateDefinitions.Add("MDL_MSVC_EDITCONTINUE_WORKAROUND=1");
 				}
 			}
-			else if (Target.Platform == UnrealTargetPlatform.Mac)
+			else if (Target.Platform == UnrealTargetPlatform.Mac || Target.Platform == UnrealTargetPlatform.Linux)
             {
                 RuntimeModuleNames.Add("dds.so");
                 RuntimeModuleNames.Add("libmdl_sdk.so");
