Add the check that the pointer is not null to avoid crash

diff --git a/Engine/Plugins/Enterprise/MDLImporter/Source/MDLImporter/Private/mdl/ApiContext.cpp b/Engine/Plugins/Enterprise/MDLImporter/Source/MDLImporter/Private/mdl/ApiContext.cpp
--- a/Engine/Plugins/Enterprise/MDLImporter/Source/MDLImporter/Private/mdl/ApiContext.cpp
+++ b/Engine/Plugins/Enterprise/MDLImporter/Source/MDLImporter/Private/mdl/ApiContext.cpp
@@ -117,7 +117,19 @@ namespace Mdl
 
 		FString FilePath = FPaths::Combine(LibrariesPath, TEXT("libmdl_sdk" MI_BASE_DLL_FILE_EXT));
 		DsoHandle        = FPlatformProcess::GetDllHandle(*FilePath);
+
+		if (!DsoHandle)
+		{
+			UE_LOG(LogMDLImporter, Error, TEXT("Error: Library %s does not exist in the path %s"), TEXT("libmdl_sdk" MI_BASE_DLL_FILE_EXT), *LibrariesPath);
+			return false;
+		}
+
 		void* SymbolPtr  = FPlatformProcess::GetDllExport(DsoHandle, TEXT("mi_factory"));
+		if (!SymbolPtr)
+		{
+			UE_LOG(LogMDLImporter, Error, TEXT("Error: Symbol %s does not exist in the library %s"), TEXT("mi_factory"), TEXT("libmdl_sdk" MI_BASE_DLL_FILE_EXT));
+			return false;
+		}
 
 		mi::neuraylib::INeuray* Neuray = mi::neuraylib::mi_factory<mi::neuraylib::INeuray>(SymbolPtr);
 		if (!Neuray)
