[MDLImporter] Fixed crash with baked textures

- replaced use of FTextureSource with FImage. FTextureSource is failing on check when created without a texture.

diff --git a/Engine/Plugins/Enterprise/MDLImporter/Source/MDLImporter/Private/MDLMaterialPropertyFactory.cpp b/Engine/Plugins/Enterprise/MDLImporter/Source/MDLImporter/Private/MDLMaterialPropertyFactory.cpp
--- a/Engine/Plugins/Enterprise/MDLImporter/Source/MDLImporter/Private/MDLMaterialPropertyFactory.cpp
+++ b/Engine/Plugins/Enterprise/MDLImporter/Source/MDLImporter/Private/MDLMaterialPropertyFactory.cpp
@@ -331,7 +331,7 @@ Mat::FParameterMap FMDLMaterialPropertyFactory::CreateProperties(EObjectFlags Fl
 			}
 			case FMDLMaterialProperty::EPropertyType::Texture:
 			{
-				Generator::FMaterialTextureFactory::FTextureSourcePtr ptr = const_cast<FTextureSource*>(Property.Texture.Source);
+				Generator::FMaterialTextureFactory::FTextureSourcePtr ptr = const_cast<FImage*>(Property.Texture.Source);
 				if (UTexture2D* Texture = TextureFactory->CreateTexture(ParentPackage, Property.Texture, ptr, Flags))
 				{
 					const bool bIsVirtualTexture = Texture->VirtualTextureStreaming;
diff --git a/Engine/Plugins/Enterprise/MDLImporter/Source/MDLImporter/Private/common/TextureProperty.h b/Engine/Plugins/Enterprise/MDLImporter/Source/MDLImporter/Private/common/TextureProperty.h
--- a/Engine/Plugins/Enterprise/MDLImporter/Source/MDLImporter/Private/common/TextureProperty.h
+++ b/Engine/Plugins/Enterprise/MDLImporter/Source/MDLImporter/Private/common/TextureProperty.h
@@ -23,7 +23,7 @@ namespace Common
 		//@note Can be a texture name or file path(i.e. also contains extension and it's disk path),
 		// if it's a file path then texture will imported from a file otherwise from texture source/data.
 		FString         Path;
-		FTextureSource* Source;
+		FImage* Source;
 
 		ETextureMode               Mode;
 		TextureMipGenSettings      MipGenSettings;
diff --git a/Engine/Plugins/Enterprise/MDLImporter/Source/MDLImporter/Private/common/TextureSource.h b/Engine/Plugins/Enterprise/MDLImporter/Source/MDLImporter/Private/common/TextureSource.h
--- a/Engine/Plugins/Enterprise/MDLImporter/Source/MDLImporter/Private/common/TextureSource.h
+++ b/Engine/Plugins/Enterprise/MDLImporter/Source/MDLImporter/Private/common/TextureSource.h
@@ -19,7 +19,7 @@ namespace Common
 	}
 
 	// Converts a float texture to a corresponding texture source.
-	inline FTextureSource* CreateTextureSource(const float* InData, int InWidth, int InHeight, int InChannels, bool bFlipY)
+	inline FImage* CreateTextureSource(const float* InData, int InWidth, int InHeight, int InChannels, bool bFlipY)
 	{
 		// @todo Oodle : use FImageView / CopyImage
 		//  just make an FImageView on the float * and use TextureSource->Init() from ImageView
@@ -29,13 +29,14 @@ namespace Common
 		// just use F32 ETextureSourceFormat ?
 
 		// use 16 bpp for linear textures support(i.e. if more than 1 channel)
-		const ETextureSourceFormat Format = InChannels == 1 ? TSF_G8 : TSF_RGBA16;
+		const ERawImageFormat::Type Format = InChannels == 1 ? ERawImageFormat::G8 : ERawImageFormat::RGBA16;
 		const int                  Size   = InWidth * InHeight * InChannels;
 
-		FTextureSource* Source = new FTextureSource();
-		Source->Init(InWidth, InHeight, 1, 1, Format);
+		FImage* Source = new FImage();
+		Source->Init(InWidth, InHeight, 1, Format);
 
-		uint8*      DstBuf = Source->LockMip(0);
+		FImageView SourceView = *Source;
+		uint8*      DstBuf = static_cast<uint8*>(SourceView.RawData);
 		const float Max8   = TNumericLimits<uint8>::Max();
 		const float Max16  = TNumericLimits<uint16>::Max();
 		switch (InChannels)
@@ -106,7 +107,7 @@ namespace Common
 				check(false);
 				break;
 		}
-		Source->UnlockMip(0);
+		
 		return Source;
 	}
 }
diff --git a/Engine/Plugins/Enterprise/MDLImporter/Source/MDLImporter/Private/generator/MaterialTextureFactory.cpp b/Engine/Plugins/Enterprise/MDLImporter/Source/MDLImporter/Private/generator/MaterialTextureFactory.cpp
--- a/Engine/Plugins/Enterprise/MDLImporter/Source/MDLImporter/Private/generator/MaterialTextureFactory.cpp
+++ b/Engine/Plugins/Enterprise/MDLImporter/Source/MDLImporter/Private/generator/MaterialTextureFactory.cpp
@@ -52,9 +52,9 @@ namespace Generator
 		{
 			if (Source != nullptr)
 			{
-				check(Source->GetSizeX() > 4 && Source->GetSizeY() > 4);
+				check(Source->GetWidth() > 4 && Source->GetHeight() > 4);
 				Texture = Factory->CreateTexture2D(ParentPackage, *TextureName, Flags);
-				Swap(Texture->Source, *Source);  // release ownership of the texture source
+				Texture->Source.Init(*Source);
 				Texture->Source.Compress();
 
 				delete Source;
diff --git a/Engine/Plugins/Enterprise/MDLImporter/Source/MDLImporter/Private/generator/MaterialTextureFactory.h b/Engine/Plugins/Enterprise/MDLImporter/Source/MDLImporter/Private/generator/MaterialTextureFactory.h
--- a/Engine/Plugins/Enterprise/MDLImporter/Source/MDLImporter/Private/generator/MaterialTextureFactory.h
+++ b/Engine/Plugins/Enterprise/MDLImporter/Source/MDLImporter/Private/generator/MaterialTextureFactory.h
@@ -10,14 +10,14 @@
 class UTextureFactory;
 class UObject;
 class UTexture2D;
-struct FTextureSource;
+struct FImage;
 
 namespace Generator
 {
 	class FMaterialTextureFactory
 	{
 	public:
-		typedef FTextureSource* FTextureSourcePtr;
+		typedef FImage* FTextureSourcePtr;
 
 		FMaterialTextureFactory();
 
diff --git a/Engine/Plugins/Enterprise/MDLImporter/Source/MDLImporter/Private/mdl/BakeParam.h b/Engine/Plugins/Enterprise/MDLImporter/Source/MDLImporter/Private/mdl/BakeParam.h
--- a/Engine/Plugins/Enterprise/MDLImporter/Source/MDLImporter/Private/mdl/BakeParam.h
+++ b/Engine/Plugins/Enterprise/MDLImporter/Source/MDLImporter/Private/mdl/BakeParam.h
@@ -9,7 +9,7 @@
 #include "Math/Color.h"
 #include "Templates/Function.h"
 
-struct FTextureSource;
+struct FImage;
 namespace mi
 {
 	class IData;
@@ -37,7 +37,7 @@ namespace Mdl
 	class FBakeTextureData
 	{
 	public:
-		typedef FTextureSource* FTextureSourcePtr;
+		typedef FImage* FTextureSourcePtr;
 
 		FBakeTextureData(FString& Name, FTextureSourcePtr& Source);
 
diff --git a/Engine/Plugins/Enterprise/MDLImporter/Source/MDLImporter/Private/mdl/Material.h b/Engine/Plugins/Enterprise/MDLImporter/Source/MDLImporter/Private/mdl/Material.h
--- a/Engine/Plugins/Enterprise/MDLImporter/Source/MDLImporter/Private/mdl/Material.h
+++ b/Engine/Plugins/Enterprise/MDLImporter/Source/MDLImporter/Private/mdl/Material.h
@@ -11,7 +11,7 @@
 #include "Templates/Function.h"
 
 class UMaterialExpression;
-struct FTextureSource;
+struct FImage;
 namespace mi
 {
 	namespace neuraylib
