[MDL Importer] Fixed crash importing Basic_Material.mdl

- fixed Min support for non-equal arg counts
- fixed vector conversion from vecrtor parameter
- replaced more checks with ensure

diff --git a/Engine/Plugins/Enterprise/MDLImporter/Source/MDLImporter/Private/generator/GeneratorMaterialExpressions.cpp b/Engine/Plugins/Enterprise/MDLImporter/Source/MDLImporter/Private/generator/GeneratorMaterialExpressions.cpp
--- a/Engine/Plugins/Enterprise/MDLImporter/Source/MDLImporter/Private/generator/GeneratorMaterialExpressions.cpp
+++ b/Engine/Plugins/Enterprise/MDLImporter/Source/MDLImporter/Private/generator/GeneratorMaterialExpressions.cpp
@@ -399,7 +399,7 @@ namespace Generator
 			UMaterialExpressionAdd* Add = Cast<UMaterialExpressionAdd>(Expression);
 			const uint32 ACount = Add->A.Expression ? ComponentCount(Add->A.Expression, Add->A.OutputIndex) : 1;
 			const uint32 BCount = Add->B.Expression ? ComponentCount(Add->B.Expression, Add->B.OutputIndex) : 1;
-			check((ACount == 1) || (BCount == 1) || (ACount == BCount));
+			ensure((ACount == 1) || (BCount == 1) || (ACount == BCount));
 			return FMath::Max(ACount, BCount);
 		}
 		else if (Expression->IsA<UMaterialExpressionAppendVector>())
@@ -411,7 +411,7 @@ namespace Generator
 		else if (Expression->IsA<UMaterialExpressionArccosine>())
 		{
 			// arccosine always returns a float... check that it also just gets one!
-			check(ComponentCount(Cast<UMaterialExpressionArccosine>(Expression)->Input.Expression,
+			ensure(ComponentCount(Cast<UMaterialExpressionArccosine>(Expression)->Input.Expression,
 								  OutputIndex) == 1);
 			return 1;
 		}
@@ -419,20 +419,20 @@ namespace Generator
 		{
 			// arctangent2 always returns a float... check that it also just gets one!
 			UMaterialExpressionArctangent2* ArcTangent2 = Cast<UMaterialExpressionArctangent2>(Expression);
-			check(ComponentCount(ArcTangent2->Y.Expression, ArcTangent2->Y.OutputIndex) == 1);
-			check(ComponentCount(ArcTangent2->X.Expression, ArcTangent2->X.OutputIndex) == 1);
+			ensure(ComponentCount(ArcTangent2->Y.Expression, ArcTangent2->Y.OutputIndex) == 1);
+			ensure(ComponentCount(ArcTangent2->X.Expression, ArcTangent2->X.OutputIndex) == 1);
 			return 1;
 		}
 		else if (Expression->IsA<UMaterialExpressionBlackBody>())
 		{
 			// BlackBody always returns a color, and digests a float
 			UMaterialExpressionBlackBody* BlackBody = Cast<UMaterialExpressionBlackBody>(Expression);
-			check(ComponentCount(BlackBody->Temp.Expression, BlackBody->Temp.OutputIndex) == 1);
+			ensure(ComponentCount(BlackBody->Temp.Expression, BlackBody->Temp.OutputIndex) == 1);
 			return 3;
 		}
 		else if (Expression->IsA<UMaterialExpressionBreakMaterialAttributes>())
 		{
-			check(OutputIndex < Expression->Outputs.Num());
+			ensure(OutputIndex < Expression->Outputs.Num());
 			FExpressionOutput const& Output = Expression->Outputs[OutputIndex];
 			return Output.MaskR + Output.MaskG + Output.MaskB + Output.MaskA;
 		}
@@ -469,8 +469,8 @@ namespace Generator
 		}
 		else if (Expression->IsA<UMaterialExpressionCosine>())
 		{
-			// cosine always returns a float... check that it also just gets one!
-			check(ComponentCount(
+			// cosine always returns a float... ensure that it also just gets one!
+			ensure(ComponentCount(
 				Cast<UMaterialExpressionCosine>(Expression)->Input.Expression, OutputIndex) == 1);
 			return 1;
 		}
@@ -478,8 +478,8 @@ namespace Generator
 		{
 			// cross product always returns a float3... check that it also just gets such!
 			UMaterialExpressionCrossProduct* CrossProduct = Cast<UMaterialExpressionCrossProduct>(Expression);
-			check(ComponentCount(CrossProduct->A.Expression, CrossProduct->A.OutputIndex) == 3);
-			check(ComponentCount(CrossProduct->B.Expression, CrossProduct->B.OutputIndex) == 3);
+			ensure(ComponentCount(CrossProduct->A.Expression, CrossProduct->A.OutputIndex) == 3);
+			ensure(ComponentCount(CrossProduct->B.Expression, CrossProduct->B.OutputIndex) == 3);
 			return 3;
 		}
 		else if (Expression->IsA<UMaterialExpressionCustom>())
@@ -496,7 +496,7 @@ namespace Generator
 			case CMOT_Float4:
 				return 4;
 			default:
-				check(false);
+				ensure(false);
 				return 0;
 			}
 		}
@@ -505,7 +505,7 @@ namespace Generator
 			UMaterialExpressionDistance* Distance = Cast<UMaterialExpressionDistance>(Expression);
 			const uint32 ACount = ComponentCount(Distance->A.Expression, Distance->A.OutputIndex);
 			const uint32 BCount = ComponentCount(Distance->B.Expression, Distance->B.OutputIndex);
-			check((1 == ACount) || (1 == BCount) || (ACount == BCount));
+			ensure((1 == ACount) || (1 == BCount) || (ACount == BCount));
 			return 1;
 		}
 		else if (Expression->IsA<UMaterialExpressionDivide>())
@@ -513,7 +513,7 @@ namespace Generator
 			UMaterialExpressionDivide* Divide = Cast<UMaterialExpressionDivide>(Expression);
 			const uint32 ACount = Divide->A.Expression ? ComponentCount(Divide->A.Expression, Divide->A.OutputIndex) : 1;
 			const uint32 BCount = Divide->B.Expression ? ComponentCount(Divide->B.Expression, Divide->B.OutputIndex) : 1;
-			check((1 == ACount) || (1 == BCount) || (ACount == BCount));
+			ensure((1 == ACount) || (1 == BCount) || (ACount == BCount));
 			return FMath::Max(ACount, BCount);
 		}
 		else if (Expression->IsA<UMaterialExpressionDotProduct>())
@@ -529,7 +529,7 @@ namespace Generator
 		{
 			UMaterialExpressionFmod* Fmod = Cast<UMaterialExpressionFmod>(Expression);
 			const uint32 ACount = ComponentCount(Fmod->A.Expression, Fmod->A.OutputIndex);
-			check(ACount == ComponentCount(Fmod->B.Expression, Fmod->B.OutputIndex));
+			ensure(ACount == ComponentCount(Fmod->B.Expression, Fmod->B.OutputIndex));
 			return ACount;
 		}
 		else if (Expression->IsA<UMaterialExpressionFrac>())
@@ -559,7 +559,7 @@ namespace Generator
 			case FunctionInput_StaticBool:
 			case FunctionInput_MaterialAttributes:
 			default:
-				check(false);
+				ensure(false);
 				return 0;
 			}
 		}
@@ -567,7 +567,7 @@ namespace Generator
 		{
 			UMaterialExpressionIf* If = Cast<UMaterialExpressionIf>(Expression);
 			const uint32 ALessThanBCount = ComponentCount(If->ALessThanB.Expression, If->ALessThanB.OutputIndex);
-			check((!If->AEqualsB.Expression || (ALessThanBCount == ComponentCount(If->AEqualsB.Expression, If->AEqualsB.OutputIndex))) &&
+			ensure((!If->AEqualsB.Expression || (ALessThanBCount == ComponentCount(If->AEqualsB.Expression, If->AEqualsB.OutputIndex))) &&
 				(ALessThanBCount == ComponentCount(If->AGreaterThanB.Expression, If->AGreaterThanB.OutputIndex)));
 			return ALessThanBCount;
 		}
@@ -579,7 +579,7 @@ namespace Generator
 				LinearInterpolate->A.Expression ? ComponentCount(LinearInterpolate->A.Expression, LinearInterpolate->A.OutputIndex) : 1;
 			const uint32 BCount =
 				LinearInterpolate->B.Expression ? ComponentCount(LinearInterpolate->B.Expression, LinearInterpolate->B.OutputIndex) : 1;
-			check(ACount == BCount);
+			ensure(ACount == BCount);
 			return ACount;
 		}
 		else if (Expression->IsA<UMaterialExpressionLogarithm2>())
@@ -591,7 +591,7 @@ namespace Generator
 		{
 			UMaterialExpressionMaterialFunctionCall* MaterialFunctionCall =
 				Cast<UMaterialExpressionMaterialFunctionCall>(Expression);
-			check(MaterialFunctionCall->FunctionOutputs[OutputIndex].ExpressionOutput->A.Expression);
+			ensure(MaterialFunctionCall->FunctionOutputs[OutputIndex].ExpressionOutput->A.Expression);
 			return ComponentCount(MaterialFunctionCall->FunctionOutputs[OutputIndex].ExpressionOutput->A.Expression,
 								   MaterialFunctionCall->FunctionOutputs[OutputIndex].ExpressionOutput->A.OutputIndex);
 		}
@@ -600,22 +600,24 @@ namespace Generator
 			UMaterialExpressionMax* Max = Cast<UMaterialExpressionMax>(Expression);
 			const uint32 ACount = Max->A.Expression ? ComponentCount(Max->A.Expression, Max->A.OutputIndex) : 1;
 			const uint32 BCount = Max->B.Expression ? ComponentCount(Max->B.Expression, Max->B.OutputIndex) : 1;
-			check((ACount == 1) || (BCount == 1) || (ACount == BCount));
+			ensure((ACount == 1) || (BCount == 1) || (ACount == BCount));
 			return FMath::Max(ACount, BCount);
 		}
 		else if (Expression->IsA<UMaterialExpressionMin>())
 		{
 			UMaterialExpressionMin* Min = Cast<UMaterialExpressionMin>(Expression);
 			const uint32 ACount = Min->A.Expression ? ComponentCount(Min->A.Expression, Min->A.OutputIndex) : 1;
-			check(ACount == (Min->B.Expression ? ComponentCount(Min->B.Expression, Min->B.OutputIndex) : 1));
-			return ACount;
+			const uint32 BCount = Min->B.Expression ? ComponentCount(Min->B.Expression, Min->B.OutputIndex) : 1;
+			ensure((ACount == 1) || (BCount == 1) || (ACount == BCount));
+
+			return FMath::Max(ACount, BCount);
 		}
 		else if (Expression->IsA<UMaterialExpressionMultiply>())
 		{
 			UMaterialExpressionMultiply* Multiply = Cast<UMaterialExpressionMultiply>(Expression);
 			const uint32 ACount = Multiply->A.Expression ? ComponentCount(Multiply->A.Expression, Multiply->A.OutputIndex) : 1;
 			const uint32 BCount = Multiply->B.Expression ? ComponentCount(Multiply->B.Expression, Multiply->B.OutputIndex) : 1;
-			check((ACount == 1) || (BCount == 1) || (ACount == BCount));
+			ensure((ACount == 1) || (BCount == 1) || (ACount == BCount));
 			return FMath::Max(ACount, BCount);
 		}
 		else if (Expression->IsA<UMaterialExpressionNormalize>())
@@ -639,15 +641,15 @@ namespace Generator
 		}
 		else if (Expression->IsA<UMaterialExpressionSine>())
 		{
-			// sine always returns a float... check that it also just gets one!
-			check(ComponentCount(
+			// sine always returns a float... ensure that it also just gets one!
+			ensure(ComponentCount(
 				Cast<UMaterialExpressionSine>(Expression)->Input.Expression, OutputIndex) == 1);
 			return 1;
 		}
 		else if (Expression->IsA<UMaterialExpressionSphereMask>())
 		{
 			UMaterialExpressionSphereMask* SphereMask = Cast<UMaterialExpressionSphereMask>(Expression);
-			check(ComponentCount(SphereMask->A.Expression, SphereMask->A.OutputIndex) ==
+			ensure(ComponentCount(SphereMask->A.Expression, SphereMask->A.OutputIndex) ==
 				ComponentCount(SphereMask->B.Expression, SphereMask->B.OutputIndex));
 			return 1;
 		}
@@ -669,7 +671,7 @@ namespace Generator
 			UMaterialExpressionStaticSwitch* StaticSwitch = Cast<UMaterialExpressionStaticSwitch>(Expression);
 			const uint32 ACount = ComponentCount(StaticSwitch->A.Expression, StaticSwitch->A.OutputIndex);
 			const uint32 BCount = ComponentCount(StaticSwitch->B.Expression, StaticSwitch->B.OutputIndex);
-			check(ACount == BCount);
+			ensure(ACount == BCount);
 			return ACount;
 		}
 		else if (Expression->IsA<UMaterialExpressionSubtract>())
@@ -677,7 +679,7 @@ namespace Generator
 			UMaterialExpressionSubtract* Subtract = Cast<UMaterialExpressionSubtract>(Expression);
 			const uint32 ACount = Subtract->A.Expression ? ComponentCount(Subtract->A.Expression, Subtract->A.OutputIndex) : 1;
 			const uint32 BCount = Subtract->B.Expression ? ComponentCount(Subtract->B.Expression, Subtract->B.OutputIndex) : 1;
-			check((ACount == 1) || (BCount == 1) || (ACount == BCount));
+			ensure((ACount == 1) || (BCount == 1) || (ACount == BCount));
 			return FMath::Max(ACount, BCount);
 		}
 		else if (Expression->IsA<UMaterialExpressionTextureCoordinate>())
@@ -713,8 +715,8 @@ namespace Generator
 		}
 		else if (Expression->IsA<UMaterialExpressionArcsine>())
 		{
-			// arcsine always returns a float... check that it also just gets one!
-			check(ComponentCount(Cast<UMaterialExpressionArcsine>(Expression)->Input.Expression, OutputIndex) == 1);
+			// arcsine always returns a float... ensure that it also just gets one!
+			ensure(ComponentCount(Cast<UMaterialExpressionArcsine>(Expression)->Input.Expression, OutputIndex) == 1);
 			return 1;
 		}
 		else if (Expression->IsA<UMaterialExpressionVertexNormalWS>()
diff --git a/Engine/Plugins/Enterprise/MDLImporter/Source/MDLImporter/Private/generator/MaterialExpressionFactory.cpp b/Engine/Plugins/Enterprise/MDLImporter/Source/MDLImporter/Private/generator/MaterialExpressionFactory.cpp
--- a/Engine/Plugins/Enterprise/MDLImporter/Source/MDLImporter/Private/generator/MaterialExpressionFactory.cpp
+++ b/Engine/Plugins/Enterprise/MDLImporter/Source/MDLImporter/Private/generator/MaterialExpressionFactory.cpp
@@ -1336,9 +1336,8 @@ namespace Generator
 					{
 						return {NewMaterialExpressionStaticSwitch(CurrentMaterial, Inputs[0], 1.0f, 0.0f)};
 					}
-					else
+					if (Inputs[0].IsExpressionA<UMaterialExpressionScalarParameter>())
 					{
-						check(Inputs[0].IsExpressionA<UMaterialExpressionScalarParameter>());
 						return Inputs;
 					}
 				}
@@ -1352,10 +1351,14 @@ namespace Generator
 
 				if (Inputs.Num() == 1)
 				{
-					if (ensure(IsScalar(Inputs[0])))
+					if (IsScalar(Inputs[0]))
 					{
 						return { NewMaterialExpressionFunctionCall(CurrentMaterial, MakeFloat3, {Inputs[0], Inputs[0], Inputs[0]}) };
 					}
+					if (Inputs[0].IsExpressionA<UMaterialExpressionVectorParameter>())
+					{
+						return Inputs;
+					}
 				}
 				else if (Inputs.Num() == VectorSize)
 				{
