Fixed IncludeTool static analysis warnings

diff --git a/Engine/Plugins/Enterprise/MDLImporter/Source/MDLImporter/Private/mdl/MaterialDistiller.h b/Engine/Plugins/Enterprise/MDLImporter/Source/MDLImporter/Private/mdl/MaterialDistiller.h
--- a/Engine/Plugins/Enterprise/MDLImporter/Source/MDLImporter/Private/mdl/MaterialDistiller.h
+++ b/Engine/Plugins/Enterprise/MDLImporter/Source/MDLImporter/Private/mdl/MaterialDistiller.h
@@ -5,8 +5,6 @@
 
 #include "Math/IntPoint.h"
 
-namespace Mdl { class IMapDistilHandler; }
-
 #ifdef USE_MDLSDK
 
 #include "mdl/Common.h"
@@ -18,6 +16,7 @@ namespace Mdl
 {
 	class FApiContext;
 	class FMaterialCollection;
+	class IMapDistilHandler;
 	struct FBakeParam;
 	struct FMaterial;
 
diff --git a/Engine/Plugins/Experimental/AppleVision/Source/AppleVision/Private/AppleVisionPlugin.cpp b/Engine/Plugins/Experimental/AppleVision/Source/AppleVision/Private/AppleVisionPlugin.cpp
--- a/Engine/Plugins/Experimental/AppleVision/Source/AppleVision/Private/AppleVisionPlugin.cpp
+++ b/Engine/Plugins/Experimental/AppleVision/Source/AppleVision/Private/AppleVisionPlugin.cpp
@@ -1,6 +1,5 @@
 // Copyright Epic Games, Inc. All Rights Reserved.
 
-#include "AppleVisionAvailability.h"
 #include "IAppleVisionPlugin.h"
 #include "AppleImageUtilsTypes.h"
 #include "UObject/GCObject.h"
