[Engine/Plugins]
* Fixed non unity/pch compile errors

diff --git a/Engine/Plugins/Compositing/OpenCVLensDistortion/Source/OpenCVLensDistortion/Private/OpenCVLensDistortionParameters.cpp b/Engine/Plugins/Compositing/OpenCVLensDistortion/Source/OpenCVLensDistortion/Private/OpenCVLensDistortionParameters.cpp
--- a/Engine/Plugins/Compositing/OpenCVLensDistortion/Source/OpenCVLensDistortion/Private/OpenCVLensDistortionParameters.cpp
+++ b/Engine/Plugins/Compositing/OpenCVLensDistortion/Source/OpenCVLensDistortion/Private/OpenCVLensDistortionParameters.cpp
@@ -12,6 +12,8 @@
 #include "opencv2/imgproc.hpp"
 #include "PostOpenCVHeaders.h"
 // IWYU pragma: end_keep
+#else
+#include "IOpenCVLensDistortionModule.h"
 #endif
 
 
diff --git a/Engine/Plugins/Enterprise/MDLImporter/Source/MDLImporter/Private/generator/MaterialExpressionFactory.cpp b/Engine/Plugins/Enterprise/MDLImporter/Source/MDLImporter/Private/generator/MaterialExpressionFactory.cpp
--- a/Engine/Plugins/Enterprise/MDLImporter/Source/MDLImporter/Private/generator/MaterialExpressionFactory.cpp
+++ b/Engine/Plugins/Enterprise/MDLImporter/Source/MDLImporter/Private/generator/MaterialExpressionFactory.cpp
@@ -30,6 +30,7 @@ MDLSDK_INCLUDES_START
 #include "mi/neuraylib/itransaction.h"
 #include "Materials/MaterialExpressionConstant.h"
 #include "mi/neuraylib/istring.h"
+#include "Materials/MaterialExpressionChannelMaskParameter.h"
 #include "Materials/MaterialExpressionCosine.h"
 #include "Materials/MaterialExpressionCrossProduct.h"
 #include "Materials/MaterialExpressionDistance.h"
diff --git a/Engine/Plugins/Experimental/AppleVision/Source/AppleVision/Private/AppleVisionPlugin.cpp b/Engine/Plugins/Experimental/AppleVision/Source/AppleVision/Private/AppleVisionPlugin.cpp
--- a/Engine/Plugins/Experimental/AppleVision/Source/AppleVision/Private/AppleVisionPlugin.cpp
+++ b/Engine/Plugins/Experimental/AppleVision/Source/AppleVision/Private/AppleVisionPlugin.cpp
@@ -6,6 +6,7 @@
 #include "Engine/Texture.h"
 
 #if SUPPORTS_APPLE_VISION_1_0
+	#include "Async/Async.h"
 	#import <Vision/Vision.h>
 
 	// For runtime checks so that clang doesn't warn on targets < our SDK version
diff --git a/Engine/Plugins/Experimental/ChaosFlesh/Source/ChaosFlesh/Private/ChaosFlesh/PB.cpp b/Engine/Plugins/Experimental/ChaosFlesh/Source/ChaosFlesh/Private/ChaosFlesh/PB.cpp
--- a/Engine/Plugins/Experimental/ChaosFlesh/Source/ChaosFlesh/Private/ChaosFlesh/PB.cpp
+++ b/Engine/Plugins/Experimental/ChaosFlesh/Source/ChaosFlesh/Private/ChaosFlesh/PB.cpp
@@ -3,6 +3,7 @@
 #include "ChaosFlesh/FleshCollectionUtility.h" // for LogChaosFlesh
 #include "ChaosFlesh/IFileStream.h"
 #include "HAL/PlatformFile.h"
+#include "HAL/PlatformFileManager.h"
 
 #ifdef USE_ZLIB
 #include "ChaosFlesh/ZIP.h"
diff --git a/Engine/Plugins/Online/OnlineServicesEOSGS/Source/Private/Online/OnlineServicesEOSGSPlatformFactory.cpp b/Engine/Plugins/Online/OnlineServicesEOSGS/Source/Private/Online/OnlineServicesEOSGSPlatformFactory.cpp
--- a/Engine/Plugins/Online/OnlineServicesEOSGS/Source/Private/Online/OnlineServicesEOSGSPlatformFactory.cpp
+++ b/Engine/Plugins/Online/OnlineServicesEOSGS/Source/Private/Online/OnlineServicesEOSGSPlatformFactory.cpp
@@ -2,6 +2,7 @@
 
 #include "Online/OnlineServicesEOSGSPlatformFactory.h"
 
+#include "EOSShared.h"
 #include "Online/OnlineServicesEOSGS.h"
 
 #include "Misc/ConfigCacheIni.h"
diff --git a/Engine/Plugins/Online/OnlineSubsystemSteam/Source/Private/OnlineIdentityInterfaceSteam.cpp b/Engine/Plugins/Online/OnlineSubsystemSteam/Source/Private/OnlineIdentityInterfaceSteam.cpp
--- a/Engine/Plugins/Online/OnlineSubsystemSteam/Source/Private/OnlineIdentityInterfaceSteam.cpp
+++ b/Engine/Plugins/Online/OnlineSubsystemSteam/Source/Private/OnlineIdentityInterfaceSteam.cpp
@@ -3,6 +3,7 @@
 #include "OnlineIdentityInterfaceSteam.h"
 
 #include "Misc/ConfigCacheIni.h"
+#include "OnlineEncryptedAppTicketInterfaceSteam.h"
 #include "OnlineSubsystemSteam.h"
 #include "OnlineSubsystemSteamTypes.h"
 #include "OnlineError.h"
diff --git a/Engine/Plugins/Runtime/AR/Google/GoogleARCoreServices/Source/GoogleARCoreServices/Private/GoogleARCoreCloudARPinManager.cpp b/Engine/Plugins/Runtime/AR/Google/GoogleARCoreServices/Source/GoogleARCoreServices/Private/GoogleARCoreCloudARPinManager.cpp
--- a/Engine/Plugins/Runtime/AR/Google/GoogleARCoreServices/Source/GoogleARCoreServices/Private/GoogleARCoreCloudARPinManager.cpp
+++ b/Engine/Plugins/Runtime/AR/Google/GoogleARCoreServices/Source/GoogleARCoreServices/Private/GoogleARCoreCloudARPinManager.cpp
@@ -4,6 +4,8 @@
 #if PLATFORM_ANDROID
 #include "arcore_c_api.h"
 #elif PLATFORM_IOS
+#include "Misc/ConfigCacheIni.h"
+#include "Misc/EngineVersion.h"
 #include "arcore_ios_c_api.h"
 #endif
 
diff --git a/Engine/Plugins/VirtualProduction/LiveLinkXR/Source/LiveLinkXR/Private/LiveLinkXRSource.cpp b/Engine/Plugins/VirtualProduction/LiveLinkXR/Source/LiveLinkXR/Private/LiveLinkXRSource.cpp
--- a/Engine/Plugins/VirtualProduction/LiveLinkXR/Source/LiveLinkXR/Private/LiveLinkXRSource.cpp
+++ b/Engine/Plugins/VirtualProduction/LiveLinkXR/Source/LiveLinkXR/Private/LiveLinkXRSource.cpp
@@ -12,6 +12,7 @@
 #include "LiveLinkXROpenXRExt.h"
 #include "LiveLinkXROpenXRExtModule.h"
 #include "LiveLinkXRConnectionSettings.h"
+#include "LiveLinkXRSourceSettings.h"
 #include "Misc/App.h"
 #include "Roles/LiveLinkTransformRole.h"
 #include "Roles/LiveLinkTransformTypes.h"
