Quick workaround for LLVM 15.0.1 not a perfect fix but gets around warning until a proper compare is put in

diff --git a/Engine/Plugins/Enterprise/MDLImporter/Source/MDLImporter/Private/mdl/MaterialPrinter.cpp b/Engine/Plugins/Enterprise/MDLImporter/Source/MDLImporter/Private/mdl/MaterialPrinter.cpp
--- a/Engine/Plugins/Enterprise/MDLImporter/Source/MDLImporter/Private/mdl/MaterialPrinter.cpp
+++ b/Engine/Plugins/Enterprise/MDLImporter/Source/MDLImporter/Private/mdl/MaterialPrinter.cpp
@@ -581,7 +581,7 @@ namespace Mdl
 						{
 							TraverseResult << TEXT(", gamma: ::tex::gamma_linear");
 						}
-						else if (Gamma == 2.2)
+						else if (Gamma == 2.2f)
 						{
 							TraverseResult << TEXT(", gamma: ::tex::gamma_srgb");
 						}
