Fixed clang warning

diff --git a/Engine/Plugins/Enterprise/MDLImporter/Source/MDLImporter/Private/mdl/MaterialPrinter.cpp b/Engine/Plugins/Enterprise/MDLImporter/Source/MDLImporter/Private/mdl/MaterialPrinter.cpp
--- a/Engine/Plugins/Enterprise/MDLImporter/Source/MDLImporter/Private/mdl/MaterialPrinter.cpp
+++ b/Engine/Plugins/Enterprise/MDLImporter/Source/MDLImporter/Private/mdl/MaterialPrinter.cpp
@@ -573,15 +573,15 @@ namespace Mdl
 					{
 						// gamma
 						const float Gamma = Texture->get_gamma();
-						if (Gamma == 0.0)
+						if (FMath::IsNearlyEqual(Gamma, 0.0f))
 						{
 							TraverseResult << TEXT(", gamma: ::tex::gamma_default");
 						}
-						else if (Gamma == 1.0)
+						else if (FMath::IsNearlyEqual(Gamma, 1.0f))
 						{
 							TraverseResult << TEXT(", gamma: ::tex::gamma_linear");
 						}
-						else if (Gamma == 2.2f)
+						else if (FMath::IsNearlyEqual(Gamma, 2.2f))
 						{
 							TraverseResult << TEXT(", gamma: ::tex::gamma_srgb");
 						}
