Fixed missing textures in content browser when loading an MDL file
Root cause: MDL importer was creating a package with all the assets in it. Material and its textures.

#rnx

diff --git a/Engine/Plugins/Enterprise/MDLImporter/Source/MDLImporter/MDLImporter.Build.cs b/Engine/Plugins/Enterprise/MDLImporter/Source/MDLImporter/MDLImporter.Build.cs
--- a/Engine/Plugins/Enterprise/MDLImporter/Source/MDLImporter/MDLImporter.Build.cs
+++ b/Engine/Plugins/Enterprise/MDLImporter/Source/MDLImporter/MDLImporter.Build.cs
@@ -22,9 +22,10 @@ namespace UnrealBuildTool.Rules
             PrivateDependencyModuleNames.AddRange(
                 new string[]
                 {
-                    "Analytics",
-                    "Core",
-                    "RenderCore",
+					"Analytics",
+					"AssetTools",
+					"Core",
+					"RenderCore",
 					"RHI",
                     "ImageCore",
                     "CoreUObject",
diff --git a/Engine/Plugins/Enterprise/MDLImporter/Source/MDLImporter/Private/generator/MaterialTextureFactory.cpp b/Engine/Plugins/Enterprise/MDLImporter/Source/MDLImporter/Private/generator/MaterialTextureFactory.cpp
--- a/Engine/Plugins/Enterprise/MDLImporter/Source/MDLImporter/Private/generator/MaterialTextureFactory.cpp
+++ b/Engine/Plugins/Enterprise/MDLImporter/Source/MDLImporter/Private/generator/MaterialTextureFactory.cpp
@@ -2,12 +2,14 @@
 
 #include "MaterialTextureFactory.h"
 
+#include "AssetRegistry/AssetRegistryModule.h"
 #include "EditorFramework/AssetImportData.h"
 #include "Engine/Texture2D.h"
 #include "Factories/TextureFactory.h"
+#include "IAssetTools.h"
 #include "Misc/Paths.h"
 #include "ObjectTools.h"
-#include "AssetRegistry/AssetRegistryModule.h"
+#include "PackageTools.h"
 
 namespace Generator
 {
@@ -19,9 +21,36 @@ namespace Generator
 	UTexture2D* FMaterialTextureFactory::CreateTexture(UObject* ParentPackage, const Common::FTextureProperty& Property, FTextureSourcePtr& Source,
 	                                                   EObjectFlags Flags, TArray<MDLImporterLogging::FLogMessage>* InLogMessages /*= nullptr*/)
 	{
-		const FString TextureName = AssetPrefix + TEXT("_") + ObjectTools::SanitizeObjectName(FPaths::GetBaseFilename(Property.Path));
-		if (TextureName.IsEmpty())
+		TFunction<void(const FString&)> LogWarning = [InLogMessages](const FString& Message) -> void
+		{
+			if (InLogMessages)
+			{
+				InLogMessages->Emplace(MDLImporterLogging::EMessageSeverity::Warning, TEXT("Texture creation failed: ") + Message);
+			}
+		};
+
+		FString TextureName = AssetPrefix + TEXT("_") + ObjectTools::SanitizeObjectName(FPaths::GetBaseFilename(Property.Path));
+		if (TextureName == TEXT("_"))
+		{
+			LogWarning(TEXT("Texture does not have a name"));
+			return nullptr;
+		}
+
+		if (!ensure(ParentPackage))
 		{
+			LogWarning(TEXT("Invalid package for texture ") + Property.Path);
+			return nullptr;
+		}
+
+		if (!ensure(Source))
+		{
+			LogWarning(TEXT("Not supported texture format '") + FPaths::GetExtension(Property.Path) + TEXT("' for ") + Property.Path);
+			return nullptr;
+		}
+
+		if (!ensure(Source->GetWidth() > 4 && Source->GetHeight() > 4))
+		{
+			LogWarning(TEXT("Not supported texture size for ") + Property.Path);
 			return nullptr;
 		}
 
@@ -29,21 +58,42 @@ namespace Generator
 		Factory->SuppressImportOverwriteDialog();
 		UpdateTextureFactorySettings(Factory, Property);
 
-		UTexture2D* Texture = nullptr;
-		const FString Extension = FPaths::GetExtension(Property.Path);
-		if (Extension == TEXT("png") || Extension == TEXT("jpg") || Extension == TEXT("jpeg"))
+		// The package of the material is passed as the parent package
+		// Therefore we need to create a package for the texture asset itself
+		// This was unnoticed so far because the ContentBrowser was displaying packages
+		// with multiple assets as separate entries in folders. This has been removed for 5.3. 
+		const FString PackagePath = FPaths::GetPath(ParentPackage->GetPathName()) / TextureName;
+		UPackage* Package = CreatePackage(*UPackageTools::SanitizePackageName(*PackagePath));
+
+		// check for asset collision
 		{
-			// check for asset collision
-			UObject* Asset = LoadObject<UObject>(ParentPackage, *TextureName, nullptr, LOAD_NoWarn);
+			UObject* Asset = LoadObject<UObject>(Package, *TextureName, nullptr, LOAD_NoWarn);
 			if (Asset)
 			{
-				check(Asset->IsA<UTexture2D>());
-				return Cast<UTexture2D>(Asset);  // already loaded
+				if (ensure((Asset->IsA<UTexture2D>())))
+				{
+					return Cast<UTexture2D>(Asset);  // already loaded
+				}
+
+				// Make sure the asset name is unique
+				FString AssetName;
+				FString NewPackagePath;
+				IAssetTools::Get().CreateUniqueAssetName(PackagePath, FString(), /*out*/ NewPackagePath, /*out*/ AssetName);
+
+				TextureName = AssetName;
+
+				Package->MarkAsGarbage();
+				Package = CreatePackage(*NewPackagePath);
 			}
+		}
 
-			check(Source == nullptr);
+
+		UTexture2D* Texture = nullptr;
+		const FString Extension = FPaths::GetExtension(Property.Path);
+		if (Extension == TEXT("png") || Extension == TEXT("jpg") || Extension == TEXT("jpeg"))
+		{
 			bool bOperationCanceled = false;
-			Texture = static_cast<UTexture2D*>(Factory->FactoryCreateFile(UTexture2D::StaticClass(), ParentPackage, *TextureName, Flags,
+			Texture = static_cast<UTexture2D*>(Factory->FactoryCreateFile(UTexture2D::StaticClass(), Package, *TextureName, Flags,
 			                                                              Property.Path, TEXT("MDL"), nullptr, bOperationCanceled));
 			Texture->AssetImportData->Update(Property.Path);
 			FAssetRegistryModule::AssetCreated(Texture);
@@ -51,23 +101,12 @@ namespace Generator
 		}
 		else
 		{
-			if (Source != nullptr)
-			{
-				check(Source->GetWidth() > 4 && Source->GetHeight() > 4);
-				Texture = Factory->CreateTexture2D(ParentPackage, *TextureName, Flags);
-				Texture->Source.Init(*Source);
-				Texture->Source.Compress();
+			Texture = Factory->CreateTexture2D(Package, *TextureName, Flags);
+			Texture->Source.Init(*Source);
+			Texture->Source.Compress();
 
-				delete Source;
-				Source = nullptr;
-			}
-			else
-			{
-				if (InLogMessages)
-				{
-					InLogMessages->Emplace(MDLImporterLogging::EMessageSeverity::Warning, TEXT("Not supported texture format '") + Extension + "' for " + Property.Path);
-				}
-			}
+			delete Source;
+			Source = nullptr;
 		}
 
 		if (Texture)
