Fix SpiralBlure - SceneTexture has border artifacts when applied to translucent materials by clamping the min max uv by the viewport bilinear filtering min max.

#tests projects inside ticket

#virtualized

diff --git a/Engine/Plugins/Enterprise/MDLImporter/Shaders/Private/SpiralBlurWithIOR.ush b/Engine/Plugins/Enterprise/MDLImporter/Shaders/Private/SpiralBlurWithIOR.ush
--- a/Engine/Plugins/Enterprise/MDLImporter/Shaders/Private/SpiralBlurWithIOR.ush
+++ b/Engine/Plugins/Enterprise/MDLImporter/Shaders/Private/SpiralBlurWithIOR.ush
@@ -99,6 +99,10 @@ if (BlurAmount == 0.0)
 else
 {
     int i = 0;
+	// If the code is applied to scene texture, need to clamp based on viewport bilinear min max
+	//float2 BilinearOffset = View.BufferSizeAndInvSize.zw * 0.5f;
+	//float2 BilinearMinBufferUV = ViewportUVToBufferUV(float2(0.0f, 0.0f)) + BilinearOffset;
+	//float2 BilinearMaxBufferUV = ViewportUVToBufferUV(float2(1.0f, 1.0f)) - BilinearOffset;
     while (i < RadialSteps)
     {
         for (int j = 0; j < AngularSteps; j++)
@@ -108,6 +112,7 @@ else
             CurOffset.y = sin(Angle);
             CurOffset *= BlurAmount * MagicBlurNumber * View.ViewToClip[0][0] * 0.01;
             NewUV = BaseUV + (CurOffset * (CurDistance + TempAADistance));
+			//NewUV = clamp(NewUV, BilinearMinBufferUV, BilinearMaxBufferUV);
             CurColor += DecodeSceneColorForMaterialNode(NewUV);
             Substep++;
         }
