MDL: Downgrade warning into a log when MDL SDK is not available as that can prevent cooking in some cases.

diff --git a/Engine/Plugins/Enterprise/MDLImporter/Source/MDLImporter/Private/MDLImporterFactory.cpp b/Engine/Plugins/Enterprise/MDLImporter/Source/MDLImporter/Private/MDLImporterFactory.cpp
--- a/Engine/Plugins/Enterprise/MDLImporter/Source/MDLImporter/Private/MDLImporterFactory.cpp
+++ b/Engine/Plugins/Enterprise/MDLImporter/Source/MDLImporter/Private/MDLImporterFactory.cpp
@@ -113,7 +113,7 @@ UMDLImporterFactory::UMDLImporterFactory(const FObjectInitializer& ObjectInitial
 		Formats.Add(Extension);
 	}
 #else
-	UE_LOG(LogMDLImporter, Warning, TEXT("MDL SDK was not available when plugin was build - MDL plugin won't be functional!"));
+	UE_LOG(LogMDLImporter, Log, TEXT("MDL SDK was not available when plugin was build - MDL plugin won't be functional!"));
 #endif
 }
 
